Add Trace Control Window in Session Frontend
Initial implementation of the TraceTools module.

diff --git a/Engine/Source/Developer/SessionFrontend/Private/Widgets/SSessionFrontend.cpp b/Engine/Source/Developer/SessionFrontend/Private/Widgets/SSessionFrontend.cpp
--- a/Engine/Source/Developer/SessionFrontend/Private/Widgets/SSessionFrontend.cpp
+++ b/Engine/Source/Developer/SessionFrontend/Private/Widgets/SSessionFrontend.cpp
@@ -15,6 +15,7 @@
 #include "IAutomationWindowModule.h"
 #include "Interfaces/IScreenShotToolsModule.h"
 #include "Interfaces/IScreenShotComparisonModule.h"
+#include "Interfaces/ITraceToolsModule.h"
 #include "ISessionServicesModule.h"
 #include "Widgets/Browser/SSessionBrowser.h"
 #include "Widgets/Console/SSessionConsole.h"
@@ -37,6 +38,8 @@ static const FName AutomationTabId("AutomationPanel");
 static const FName SessionBrowserTabId("SessionBrowser");
 static const FName SessionConsoleTabId("SessionConsole");
 static const FName SessionScreenTabId("ScreenComparison");
+static const FName TraceControlTabId("TraceControl");
+
 static const FName ProfilerTabId("Profiler");
 
 
@@ -68,6 +71,11 @@ void SSessionFrontend::Construct( const FArguments& InArgs, const TSharedRef<SDo
 		.SetIcon(FSlateIcon(FAppStyle::GetAppStyleSetName(), "SessionFrontEnd.Tabs.Console"))
 		.SetGroup(AppMenuGroup);
 
+	TabManager->RegisterTabSpawner(TraceControlTabId, FOnSpawnTab::CreateRaw(this, &SSessionFrontend::HandleTabManagerSpawnTab, TraceControlTabId))
+		.SetDisplayName(LOCTEXT("TraceControlTabTitle", "Trace Control"))
+		.SetIcon(FSlateIcon(FAppStyle::GetAppStyleSetName(), "SessionFrontEnd.Tabs.TraceControl"))
+		.SetGroup(AppMenuGroup);
+
 	TabManager->RegisterTabSpawner(SessionScreenTabId, FOnSpawnTab::CreateRaw(this, &SSessionFrontend::HandleTabManagerSpawnTab, SessionScreenTabId))
 		.SetDisplayName(LOCTEXT("ScreenTabTitle", "Screen Comparison"))
 		.SetIcon(FSlateIcon(FAppStyle::GetAppStyleSetName(), "SessionFrontEnd.Tabs.ScreenComparison"))
@@ -81,7 +89,7 @@ void SSessionFrontend::Construct( const FArguments& InArgs, const TSharedRef<SDo
 #endif
 	
 	// create tab layout
-	const TSharedRef<FTabManager::FLayout> Layout = FTabManager::NewLayout("SessionFrontendLayout_v1.2")
+	const TSharedRef<FTabManager::FLayout> Layout = FTabManager::NewLayout("SessionFrontendLayout_v1.3")
 		->AddArea
 		(
 			FTabManager::NewPrimaryArea()
@@ -99,6 +107,7 @@ void SSessionFrontend::Construct( const FArguments& InArgs, const TSharedRef<SDo
 					// applications
 					FTabManager::NewStack()
 						->AddTab(SessionConsoleTabId, ETabState::OpenedTab)
+						->AddTab(TraceControlTabId, ETabState::OpenedTab)
 						->AddTab(AutomationTabId, ETabState::OpenedTab)
 						->AddTab(SessionScreenTabId, ETabState::OpenedTab)
 #if STATS && UE_DEPRECATED_PROFILER_ENABLED
@@ -160,6 +169,7 @@ void SSessionFrontend::InitializeControllers()
 	DeviceProxyManager = TargetDeviceServicesModule.GetDeviceProxyManager();
 	SessionManager = SessionServicesModule.GetSessionManager();
 	ScreenShotManager = ScreenShotModule.GetScreenShotManager();
+	TraceController = SessionServicesModule.GetTraceController();
 }
 
 
@@ -215,6 +225,10 @@ TSharedRef<SDockTab> SSessionFrontend::HandleTabManagerSpawnTab( const FSpawnTab
 	{
 		TabWidget = SNew(SSessionConsole, SessionManager.ToSharedRef());
 	}
+	else if (TabIdentifier == TraceControlTabId)
+	{
+		TabWidget = FModuleManager::LoadModuleChecked<UE::TraceTools::ITraceToolsModule>("TraceTools").CreateTraceControlWidget(TraceController);
+	}
 	else if (TabIdentifier == SessionScreenTabId)
 	{
 		TabWidget = FModuleManager::LoadModuleChecked<IScreenShotComparisonModule>("ScreenShotComparison").CreateScreenShotComparison(
diff --git a/Engine/Source/Developer/SessionFrontend/Private/Widgets/SSessionFrontend.h b/Engine/Source/Developer/SessionFrontend/Private/Widgets/SSessionFrontend.h
--- a/Engine/Source/Developer/SessionFrontend/Private/Widgets/SSessionFrontend.h
+++ b/Engine/Source/Developer/SessionFrontend/Private/Widgets/SSessionFrontend.h
@@ -80,4 +80,7 @@ private:
 
 	/** Holds the tab manager that manages the front-end's tabs. */
 	TSharedPtr<FTabManager> TabManager;
+
+	/** Trace controller */
+	TSharedPtr<class ITraceController> TraceController;
 };
diff --git a/Engine/Source/Developer/SessionFrontend/SessionFrontend.Build.cs b/Engine/Source/Developer/SessionFrontend/SessionFrontend.Build.cs
--- a/Engine/Source/Developer/SessionFrontend/SessionFrontend.Build.cs
+++ b/Engine/Source/Developer/SessionFrontend/SessionFrontend.Build.cs
@@ -30,6 +30,7 @@ public class SessionFrontend : ModuleRules
 				"ScreenShotComparison",
 				"ScreenShotComparisonTools",
 				"TargetPlatform",
+				"TraceTools",
 				"WorkspaceMenuStructure",
 			}
 		);
diff --git a/Engine/Source/Developer/TraceTools/Private/Models/ITraceFilterPreset.h b/Engine/Source/Developer/TraceTools/Private/Models/ITraceFilterPreset.h
new file mode 100644
--- /dev/null
+++ b/Engine/Source/Developer/TraceTools/Private/Models/ITraceFilterPreset.h
@@ -0,0 +1,43 @@
+// Copyright Epic Games, Inc. All Rights Reserved.
+
+#pragma once
+
+#include "Templates/SharedPointer.h"
+
+namespace UE::TraceTools
+{
+
+class ITraceObject;
+
+/** Interface for structures representing a filtering preset */
+struct ITraceFilterPreset : public TSharedFromThis<ITraceFilterPreset>
+{
+public:
+	virtual ~ITraceFilterPreset() {}
+
+	virtual FString GetName() const = 0;
+	virtual FText GetDisplayText() const = 0;
+	virtual FText GetDescription() const = 0;
+
+	/** Returns all names which pass this preset its filter */
+	virtual void GetAllowlistedNames(TArray<FString>& OutNames) const = 0;
+	/** Whether or not this preset can be deleted */
+	virtual bool CanDelete() const = 0;
+	/** Whether or not this preset only exists locally */
+	virtual bool IsLocal() const = 0;
+
+	/** Update this preset according to the provided set of trace objects */
+	virtual void Save(const TArray<TSharedPtr<ITraceObject>>& InObjects) = 0;
+
+	virtual void Save() = 0;
+
+	virtual void Rename(const FString& InNewName) = 0;	
+	virtual bool Delete() = 0;
+
+	/** Transition preset between being local and shared, moves between local and default INI files */
+	virtual bool MakeShared() = 0;
+	virtual bool MakeLocal() = 0;
+	
+};
+
+} // namespace UE::TraceTools
\ No newline at end of file
diff --git a/Engine/Source/Developer/TraceTools/Private/Models/ITraceObject.h b/Engine/Source/Developer/TraceTools/Private/Models/ITraceObject.h
new file mode 100644
--- /dev/null
+++ b/Engine/Source/Developer/TraceTools/Private/Models/ITraceObject.h
@@ -0,0 +1,42 @@
+// Copyright Epic Games, Inc. All Rights Reserved.
+
+#pragma once
+
+#include "Templates/SharedPointer.h"
+
+namespace UE::TraceTools
+{
+
+class ITraceObject : public TSharedFromThis<ITraceObject>
+{
+public:
+	virtual ~ITraceObject() {}
+
+	/** Returns the name representing this object */
+	virtual FString GetName() const = 0;
+	/** Returns the user-facing text for this object */
+	virtual FText GetDisplayText() const = 0;
+
+	/** Set whether or not this object should (not) be filtered out */
+	virtual void SetIsFiltered(bool bState) = 0;
+
+	/** Returns whether or not this object can be changed by the user. */
+	virtual bool IsReadOnly() const = 0;
+	
+	/** Set whether or not this object is pending an update of its state (from the connect edTrace connection) */
+	virtual void SetPending() = 0;
+	
+	/** Returns whether or not this object is set to be filtered out (true = filtered; false = not filtered)*/
+	virtual bool IsFiltered() const = 0;
+
+	/** Returns whether or not this object is pending a filter-state update (SetPending previously called) */
+	virtual bool IsPending() const = 0;
+
+	/** Returns contained child object(s) */
+	virtual void GetChildren(TArray<TSharedPtr<ITraceObject>>& OutChildren) const = 0;
+
+	/** Returns string representing this object in a search (treeview filtering) */
+	virtual void GetSearchString(TArray<FString>& OutFilterStrings) const = 0;
+};
+
+} // namespace UE::TraceTools
\ No newline at end of file
diff --git a/Engine/Source/Developer/TraceTools/Private/Models/TraceChannel.cpp b/Engine/Source/Developer/TraceTools/Private/Models/TraceChannel.cpp
new file mode 100644
--- /dev/null
+++ b/Engine/Source/Developer/TraceTools/Private/Models/TraceChannel.cpp
@@ -0,0 +1,65 @@
+// Copyright Epic Games, Inc. All Rights Reserved.
+
+#include "TraceChannel.h"
+
+#include "Internationalization/Text.h"
+
+// TraceTools
+#include "Services/ISessionTraceFilterService.h"
+
+namespace UE::TraceTools
+{
+
+FTraceChannel::FTraceChannel(FString InName, FString InParentName, uint32 InHash, bool bInEnabled, bool bInReadOnly, const TArray<TSharedPtr<ITraceObject>>& InChildObjects, TSharedPtr<ISessionTraceFilterService> InFilterService) : Name(InName), ParentName(InParentName), Hash(InHash), ChildObjects(InChildObjects), bFiltered(!bInEnabled), bIsPending(false), bReadOnly(bInReadOnly), FilterService(InFilterService)
+{
+
+}
+
+FText FTraceChannel::GetDisplayText() const
+{
+	return FText::FromString(Name);
+}
+
+FString FTraceChannel::GetName() const
+{
+	return Name;
+}
+
+void FTraceChannel::SetPending()
+{
+	bIsPending = true;
+}
+
+
+bool FTraceChannel::IsReadOnly() const
+{
+	return bReadOnly;
+}
+
+void FTraceChannel::SetIsFiltered(bool bState)
+{
+	SetPending();
+	FilterService->SetObjectFilterState(Name, !bState);
+}
+
+bool FTraceChannel::IsFiltered() const
+{
+	return bFiltered;
+}
+
+bool FTraceChannel::IsPending() const
+{
+	return bIsPending;
+}
+
+void FTraceChannel::GetSearchString(TArray<FString>& OutFilterStrings) const
+{
+	OutFilterStrings.Add(Name);
+}
+
+void FTraceChannel::GetChildren(TArray<TSharedPtr<ITraceObject>>& OutChildren) const
+{
+	OutChildren.Append(ChildObjects);
+}
+
+} // namespace UE::TraceTools
\ No newline at end of file
diff --git a/Engine/Source/Developer/TraceTools/Private/Models/TraceChannel.h b/Engine/Source/Developer/TraceTools/Private/Models/TraceChannel.h
new file mode 100644
--- /dev/null
+++ b/Engine/Source/Developer/TraceTools/Private/Models/TraceChannel.h
@@ -0,0 +1,46 @@
+// Copyright Epic Games, Inc. All Rights Reserved.
+
+#pragma once
+
+#include "ITraceObject.h"
+
+namespace UE::TraceTools
+{
+
+class ISessionTraceFilterService;
+
+class FTraceChannel : public ITraceObject
+{
+public:
+	FTraceChannel(FString InName, FString InParentName, uint32 InHash, bool bInEnabled, bool bInReadOnly, const TArray<TSharedPtr<ITraceObject>>& InChildObjects, TSharedPtr<ISessionTraceFilterService> InFilterService);
+
+	/** Begin ITraceObject overrides */
+	virtual FText GetDisplayText() const override;;
+	virtual FString GetName() const;
+	virtual void SetPending() override;
+	virtual bool IsReadOnly() const;
+	virtual void SetIsFiltered(bool bState) override;
+	virtual bool IsFiltered() const override;
+	virtual bool IsPending() const override;
+	virtual void GetSearchString(TArray<FString>& OutFilterStrings) const override;
+	virtual void GetChildren(TArray<TSharedPtr<ITraceObject>>& OutChildren) const override;
+	/** End ITraceObject overrides */
+
+protected:
+	/** This channel's name */
+	FString Name;
+	/** Channel's parent (group) name */
+	FString ParentName;
+	/** Hash representing this channel uniquely */
+	uint32 Hash;
+	TArray<TSharedPtr<ITraceObject>> ChildObjects;
+	
+	/** Whether or not this channel is filtered out, true = filtered; false = not filtered */
+	bool bFiltered;
+	bool bIsPending;	
+	bool bReadOnly;
+
+	TSharedPtr<ISessionTraceFilterService> FilterService;
+};
+
+} // namespace UE::TraceTools
\ No newline at end of file
diff --git a/Engine/Source/Developer/TraceTools/Private/Models/TraceControlCommands.h b/Engine/Source/Developer/TraceTools/Private/Models/TraceControlCommands.h
new file mode 100644
--- /dev/null
+++ b/Engine/Source/Developer/TraceTools/Private/Models/TraceControlCommands.h
@@ -0,0 +1,64 @@
+// Copyright Epic Games, Inc. All Rights Reserved.
+
+#pragma once
+
+#include "Framework/Commands/InputChord.h"
+#include "Framework/Commands/Commands.h"
+#include "TraceToolsStyle.h"
+#include "Styling/AppStyle.h"
+
+#define LOCTEXT_NAMESPACE "SessionConsoleCommands"
+
+namespace UE::TraceTools
+{
+
+class FTraceControlCommands
+	: public TCommands<FTraceControlCommands>
+{
+public:
+
+	/** Default constructor. */
+	FTraceControlCommands()
+		: TCommands<FTraceControlCommands>(
+			"TraceControl",
+			NSLOCTEXT("Contexts", "TraceControl", "Trace Control"),
+			NAME_None, FTraceToolsStyle::GetStyleSetName()
+		)
+	{ }
+
+public:
+	// TCommands interface
+	virtual void RegisterCommands() override
+	{
+		UI_COMMAND(SetTraceTargetServer, "Trace Server", "Set the trace server as the trace target.", EUserInterfaceActionType::Button, FInputChord());
+		UI_COMMAND(SetTraceTargetFile, "File", "Set file as the trace target.", EUserInterfaceActionType::Button, FInputChord());
+
+		UI_COMMAND(StartTrace, "Start Trace", "Start trace on the selected session.", EUserInterfaceActionType::Button, FInputChord());
+		UI_COMMAND(StopTrace, "Stop Trace", "Stop trace on the selected session.", EUserInterfaceActionType::Button, FInputChord());
+		UI_COMMAND(TraceSnapshot, "Snapshot", "Save a trace snapshot on the selected session.", EUserInterfaceActionType::Button, FInputChord());
+
+		UI_COMMAND(PauseTrace, "Pause Trace", "Disable all the trace channels and save the channel list so it can be enabled again with the resume command.", EUserInterfaceActionType::Button, FInputChord());
+		UI_COMMAND(ResumeTrace, "Resume Trace", "Enable the channels that were enabled before trace was paused.", EUserInterfaceActionType::Button, FInputChord());
+
+		UI_COMMAND(TraceBookmark, "Bookmark", "Trace a bookmark with a predefined name.", EUserInterfaceActionType::Button, FInputChord());
+		UI_COMMAND(TraceScreenshot, "Screenshot", "Trace a screenshot with a predefined name.", EUserInterfaceActionType::Button, FInputChord());
+
+		UI_COMMAND(ToggleStatNamedEvents, "Stat Named Events", "Enables/Disables the tracing of stat named events.", EUserInterfaceActionType::ToggleButton, FInputChord());
+	}
+
+public:
+	TSharedPtr<FUICommandInfo> SetTraceTargetServer;
+	TSharedPtr<FUICommandInfo> SetTraceTargetFile;
+	TSharedPtr<FUICommandInfo> StartTrace;
+	TSharedPtr<FUICommandInfo> StopTrace;
+	TSharedPtr<FUICommandInfo> TraceSnapshot;
+	TSharedPtr<FUICommandInfo> PauseTrace;
+	TSharedPtr<FUICommandInfo> ResumeTrace;
+	TSharedPtr<FUICommandInfo> TraceBookmark;
+	TSharedPtr<FUICommandInfo> TraceScreenshot;
+	TSharedPtr<FUICommandInfo> ToggleStatNamedEvents;
+};
+
+} // namespace UE::TraceTools
+
+#undef LOCTEXT_NAMESPACE
diff --git a/Engine/Source/Developer/TraceTools/Private/Models/TraceFilterPresets.cpp b/Engine/Source/Developer/TraceTools/Private/Models/TraceFilterPresets.cpp
new file mode 100644
--- /dev/null
+++ b/Engine/Source/Developer/TraceTools/Private/Models/TraceFilterPresets.cpp
@@ -0,0 +1,280 @@
+// Copyright Epic Games, Inc. All Rights Reserved.
+
+#include "TraceFilterPresets.h"
+
+#include "GenericPlatform/GenericPlatformFile.h"
+#include "HAL/PlatformFileManager.h"
+
+// TraceTools
+#include "ITraceObject.h"
+
+#include UE_INLINE_GENERATED_CPP_BY_NAME(TraceFilterPresets)
+
+#define LOCTEXT_NAMESPACE "TraceFilterPreset"
+
+void USharedTraceFilterPresetContainer::GetSharedUserPresets(TArray<TSharedPtr<UE::TraceTools::ITraceFilterPreset>>& OutPresets)
+{
+	for (FTraceFilterData& FilterData : SharedPresets)
+	{
+		OutPresets.Add(MakeShared<UE::TraceTools::FUserFilterPreset>(FilterData.Name, FilterData));
+	}
+}
+
+void USharedTraceFilterPresetContainer::AddFilterData(const FTraceFilterData& InFilterData)
+{
+	USharedTraceFilterPresetContainer* SharedPresetsContainer = GetMutableDefault<USharedTraceFilterPresetContainer>();
+	SharedPresetsContainer->SharedPresets.Add(InFilterData);
+}
+
+bool USharedTraceFilterPresetContainer::RemoveFilterData(const FTraceFilterData& InFilterData)
+{
+	USharedTraceFilterPresetContainer* SharedPresetsContainer = GetMutableDefault<USharedTraceFilterPresetContainer>();
+	return SharedPresetsContainer->SharedPresets.RemoveSingle(InFilterData) == 1;
+}
+
+void USharedTraceFilterPresetContainer::Save()
+{
+	USharedTraceFilterPresetContainer* SharedPresetsContainer = GetMutableDefault<USharedTraceFilterPresetContainer>();
+	SharedPresetsContainer->TryUpdateDefaultConfigFile();
+}
+
+void UEngineTraceFilterPresetContainer::GetEnginePresets(TArray<TSharedPtr<UE::TraceTools::ITraceFilterPreset>>& OutPresets)
+{
+	for (FTraceFilterData& FilterData : EnginePresets)
+	{
+		OutPresets.Add(MakeShared<UE::TraceTools::FEngineFilterPreset>(FilterData.Name, FilterData));
+	}
+}
+
+void ULocalTraceFilterPresetContainer::GetUserPresets(TArray<TSharedPtr<UE::TraceTools::ITraceFilterPreset>>& OutPresets)
+{
+	for (FTraceFilterData& FilterData : UserPresets)
+	{
+		const bool bIsLocal = true;
+		OutPresets.Add(MakeShared<UE::TraceTools::FUserFilterPreset>(FilterData.Name, FilterData, bIsLocal));
+	}
+}
+
+void ULocalTraceFilterPresetContainer::AddFilterData(const FTraceFilterData& InFilterData)
+{
+	ULocalTraceFilterPresetContainer* LocalPresetsContained = GetMutableDefault<ULocalTraceFilterPresetContainer>();
+	LocalPresetsContained->UserPresets.Add(InFilterData);
+}
+
+bool ULocalTraceFilterPresetContainer::RemoveFilterData(const FTraceFilterData& InFilterData)
+{
+	ULocalTraceFilterPresetContainer* LocalPresetsContained = GetMutableDefault<ULocalTraceFilterPresetContainer>();
+	return LocalPresetsContained->UserPresets.RemoveSingle(InFilterData) == 1;
+}
+
+void ULocalTraceFilterPresetContainer::Save()
+{
+	ULocalTraceFilterPresetContainer* LocalPresetsContainer = GetMutableDefault<ULocalTraceFilterPresetContainer>();
+	LocalPresetsContainer->SaveConfig();
+}
+
+namespace UE::TraceTools
+{
+
+void FFilterPresetHelpers::CreateNewPreset(const TArray<TSharedPtr<ITraceObject>>& InObjects)
+{
+	ULocalTraceFilterPresetContainer* LocalPresetsContainer = GetMutableDefault<ULocalTraceFilterPresetContainer>();
+	USharedTraceFilterPresetContainer* SharedPresetsContainer = GetMutableDefault<USharedTraceFilterPresetContainer>();
+
+	FTraceFilterData& NewUserFilter = LocalPresetsContainer->UserPresets.AddDefaulted_GetRef();
+
+	bool bFoundValidName = false;
+	int32 ProfileAppendNum = LocalPresetsContainer->UserPresets.Num();
+	FString NewFilterName;
+	while (!bFoundValidName)
+	{
+		NewFilterName = FString::Printf(TEXT("UserPreset_%i"), ProfileAppendNum);
+
+		bool bValidName = true;
+		for (const FTraceFilterData& Filter : LocalPresetsContainer->UserPresets)
+		{
+			if (Filter.Name == NewFilterName)
+			{
+				bValidName = false;
+				break;
+			}
+		}
+
+		for (const FTraceFilterData& Filter : SharedPresetsContainer->SharedPresets)
+		{
+			if (Filter.Name == NewFilterName)
+			{
+				bValidName = false;
+				break;
+			}
+		}
+
+		if (!bValidName)
+		{
+			++ProfileAppendNum;
+		}
+
+		bFoundValidName = bValidName;
+	}
+
+	NewUserFilter.Name = NewFilterName;
+
+	TArray<FString> Names;
+	ExtractEnabledObjectNames(InObjects, Names);
+	NewUserFilter.AllowlistedNames = Names;
+
+	ULocalTraceFilterPresetContainer::Save();
+}
+
+bool FFilterPresetHelpers::CanModifySharedPreset()
+{
+	const USharedTraceFilterPresetContainer* SharedContainer = GetDefault<USharedTraceFilterPresetContainer>();
+	return !FPlatformFileManager::Get().GetPlatformFile().IsReadOnly(*SharedContainer->GetDefaultConfigFilename());
+}
+
+void FFilterPresetHelpers::ExtractEnabledObjectNames(const TArray<TSharedPtr<ITraceObject>>& InObjects, TArray<FString>& OutNames)
+{
+	for (const TSharedPtr<ITraceObject>& Object : InObjects)
+	{
+		if (!Object->IsFiltered())
+		{
+			OutNames.Add(Object->GetName());
+		}
+
+		TArray<TSharedPtr<ITraceObject>> ChildObjects;
+		Object->GetChildren(ChildObjects);
+		if(ChildObjects.Num())
+		{
+			ExtractEnabledObjectNames(ChildObjects, OutNames);
+		}
+	}
+}
+
+bool FUserFilterPreset::CanDelete() const
+{
+	return true;
+}
+
+bool FUserFilterPreset::Delete()
+{
+	bool bRemoved = false;
+
+	if (IsLocal())
+	{
+		bRemoved = ULocalTraceFilterPresetContainer::RemoveFilterData(FilterData);
+		ULocalTraceFilterPresetContainer::Save();
+	}
+	else
+	{
+		bRemoved = USharedTraceFilterPresetContainer::RemoveFilterData(FilterData);
+		USharedTraceFilterPresetContainer::Save();
+	}
+
+	return ensure(bRemoved);
+}
+
+bool FUserFilterPreset::MakeShared()
+{
+	ensure(IsLocal());
+
+	USharedTraceFilterPresetContainer::AddFilterData(FilterData);
+	ensure(ULocalTraceFilterPresetContainer::RemoveFilterData(FilterData));
+		
+	USharedTraceFilterPresetContainer::Save();
+	ULocalTraceFilterPresetContainer::Save();
+	
+	return true;
+}
+
+bool FUserFilterPreset::MakeLocal()
+{
+	ensure(!IsLocal());
+
+	ULocalTraceFilterPresetContainer::AddFilterData(FilterData);
+	ensure(USharedTraceFilterPresetContainer::RemoveFilterData(FilterData));
+	
+	USharedTraceFilterPresetContainer::Save();
+	ULocalTraceFilterPresetContainer::Save();
+
+	return true;
+}
+
+bool FUserFilterPreset::IsLocal() const
+{
+	return bIsLocalPreset;
+}
+
+void FUserFilterPreset::Save(const TArray<TSharedPtr<ITraceObject>>& InObjects)
+{
+	TArray<FString> Names;
+	FFilterPresetHelpers::ExtractEnabledObjectNames(InObjects, Names);
+
+	FilterData.AllowlistedNames = Names;
+	
+	USharedTraceFilterPresetContainer::Save();
+	ULocalTraceFilterPresetContainer::Save();
+}
+
+void FUserFilterPreset::Save()
+{
+	USharedTraceFilterPresetContainer::Save();
+	ULocalTraceFilterPresetContainer::Save();
+}
+
+FString FFilterPreset::GetName() const
+{
+	return Name;
+}
+
+FText FFilterPreset::GetDisplayText() const
+{
+	return FText::FromString(Name);
+}
+
+FText FFilterPreset::GetDescription() const
+{
+	return FText::FormatOrdered(LOCTEXT("FilterPresetDescriptionFormat", "Name: {0}\nType: {1}"), FText::FromString(Name), CanDelete() ? (IsLocal() ? LOCTEXT("LocalPreset", "Local") : LOCTEXT("SharedPreset", "Shared")) : LOCTEXT("EnginePreset", "Engine"));
+}
+
+void FFilterPreset::GetAllowlistedNames(TArray<FString>& OutNames) const
+{
+	OutNames.Append(FilterData.AllowlistedNames);
+}
+
+bool FFilterPreset::CanDelete() const
+{
+	return false;
+}
+
+void FFilterPreset::Rename(const FString& InNewName)
+{	
+	Name = InNewName;
+	FilterData.Name = InNewName;
+
+	Save();
+}
+
+bool FFilterPreset::Delete()
+{
+	return false;
+}
+
+bool FFilterPreset::MakeShared()
+{
+	return false;
+}
+
+bool FFilterPreset::MakeLocal()
+{
+	return false;
+}
+
+bool FFilterPreset::IsLocal() const
+{
+	return false;
+}
+
+} // namespace UE::TraceTools
+
+#undef LOCTEXT_NAMESPACE // "FilterPreset"
+
diff --git a/Engine/Source/Developer/TraceTools/Private/Models/TraceFilterPresets.h b/Engine/Source/Developer/TraceTools/Private/Models/TraceFilterPresets.h
new file mode 100644
--- /dev/null
+++ b/Engine/Source/Developer/TraceTools/Private/Models/TraceFilterPresets.h
@@ -0,0 +1,138 @@
+// Copyright Epic Games, Inc. All Rights Reserved.
+
+#pragma once
+
+#include "ITraceFilterPreset.h"
+
+#include "TraceFilterPresets.generated.h"
+
+namespace UE::TraceTools
+{
+struct FFilterPresetHelpers
+{
+	/** Creates a new filtering preset according to the specific object names */
+	static void CreateNewPreset(const TArray<TSharedPtr<ITraceObject>>& InObjects);
+	/** Creates a set of strings, corresponding to set of non-filtered out object as part of InObjects */
+	static void ExtractEnabledObjectNames(const TArray<TSharedPtr<ITraceObject>>& InObjects, TArray<FString>& OutNames);
+	/** Returns whether or not shared presets can be modified, requires write-flag on default confing files */
+	static bool CanModifySharedPreset();
+};
+}
+
+/** Structure representing an individual preset in configuration (ini) files */
+USTRUCT()
+struct FTraceFilterData
+{
+	GENERATED_BODY()
+
+	UPROPERTY()
+	FString Name;
+
+	UPROPERTY()
+	TArray<FString> AllowlistedNames;
+
+	bool operator==(const FTraceFilterData& Other) const
+	{
+		return Name == Other.Name && AllowlistedNames == Other.AllowlistedNames;
+	}
+};
+
+/** UObject containers for the preset data */
+UCLASS(Config = Trace)
+class ULocalTraceFilterPresetContainer : public UObject
+{
+	GENERATED_BODY()
+
+	friend struct UE::TraceTools::FFilterPresetHelpers;
+public:
+	void GetUserPresets(TArray<TSharedPtr<UE::TraceTools::ITraceFilterPreset>>& OutPresets);
+
+	static void AddFilterData(const FTraceFilterData& InFilterData);
+	static bool RemoveFilterData(const FTraceFilterData& InFilterData);
+	static void Save();
+protected:
+	UPROPERTY(Config)
+	TArray<FTraceFilterData> UserPresets;
+};
+
+UCLASS(Config = Trace, DefaultConfig)
+class USharedTraceFilterPresetContainer : public UObject
+{
+	GENERATED_BODY()
+
+	friend struct UE::TraceTools::FFilterPresetHelpers;
+public:
+	void GetSharedUserPresets(TArray<TSharedPtr<UE::TraceTools::ITraceFilterPreset>>& OutPresets);
+
+	static void AddFilterData(const FTraceFilterData& InFilterData);
+	static bool RemoveFilterData(const FTraceFilterData& InFilterData);
+	static void Save();
+protected:
+	UPROPERTY(Config)
+	TArray<FTraceFilterData> SharedPresets;
+};
+
+UCLASS(Config = Trace, DefaultConfig)
+class UEngineTraceFilterPresetContainer : public UObject
+{
+	GENERATED_BODY()
+
+	friend struct UE::TraceTools::FFilterPresetHelpers;
+public:
+	void GetEnginePresets(TArray<TSharedPtr<UE::TraceTools::ITraceFilterPreset>>& OutPresets);
+protected:
+	UPROPERTY(Config)
+	TArray<FTraceFilterData> EnginePresets;
+};
+
+namespace UE::TraceTools
+{
+
+/** Base implementation of a filter preset */
+struct FFilterPreset : public TraceTools::ITraceFilterPreset
+{
+public:
+	FFilterPreset(const FString& InName, FTraceFilterData& InFilterData) : Name(InName), FilterData(InFilterData) {}
+
+	/** Begin ITraceFilterPreset overrides */
+	virtual FString GetName() const override;
+	virtual FText GetDisplayText() const;
+	virtual FText GetDescription() const;
+	virtual void GetAllowlistedNames(TArray<FString>& OutNames) const override;
+	virtual bool CanDelete() const override;
+	virtual void Rename(const FString& InNewName) override;
+	virtual bool Delete() override;
+	virtual bool MakeShared() override;
+	virtual bool MakeLocal() override;
+	virtual bool IsLocal() const override;
+	virtual void Save(const TArray<TSharedPtr<TraceTools::ITraceObject>>& InObjects) override {}
+	virtual void Save() override {}
+	/** End ITraceFilterPreset overrides */
+protected:
+	FString Name;
+	FTraceFilterData& FilterData;
+};
+
+/** Engine level preset is simply a basic one */
+typedef FFilterPreset FEngineFilterPreset;
+
+/** User filter preset, allows for deletion / transitioning INI ownership */
+struct FUserFilterPreset : public FFilterPreset
+{
+public:
+	FUserFilterPreset(const FString& InName, FTraceFilterData& InFilterData, bool bInLocal = false) : FFilterPreset(InName, InFilterData), bIsLocalPreset(bInLocal) {}
+
+	/** Begin ITraceFilterPreset overrides */
+	virtual bool CanDelete() const override;
+	virtual bool Delete() override;
+	virtual bool MakeShared() override;
+	virtual bool MakeLocal() override;
+	virtual bool IsLocal() const override;
+	virtual void Save(const TArray<TSharedPtr<TraceTools::ITraceObject>>& InObjects) override;
+	virtual void Save() override;
+	/** End ITraceFilterPreset overrides */
+protected:
+	bool bIsLocalPreset;
+};
+
+} // namespace UE::TraceTools
\ No newline at end of file
diff --git a/Engine/Source/Developer/TraceTools/Private/Services/ISessionTraceFilterService.h b/Engine/Source/Developer/TraceTools/Private/Services/ISessionTraceFilterService.h
new file mode 100644
--- /dev/null
+++ b/Engine/Source/Developer/TraceTools/Private/Services/ISessionTraceFilterService.h
@@ -0,0 +1,45 @@
+// Copyright Epic Games, Inc. All Rights Reserved.
+
+#pragma once
+
+#include "Templates/SharedPointer.h"
+
+struct FDateTime;
+
+namespace UE::TraceTools
+{
+
+struct ITraceFilterPreset;
+
+struct FTraceObjectInfo
+{
+	FString Name;
+	bool bEnabled;
+	bool bReadOnly;
+	uint32 Hash;
+	uint32 OwnerHash;
+};
+
+/** Filtering service, representing the state and data for a specific TraceServices::IAnalysisSession */
+class ISessionTraceFilterService : public TSharedFromThis<ISessionTraceFilterService>
+{
+public:
+	virtual ~ISessionTraceFilterService() {}
+
+	/** Returns the root level set of objects */
+	virtual void GetRootObjects(TArray<FTraceObjectInfo>& OutObjects) const = 0;
+
+	/** Returns the contained objects for the specific object hash */
+	virtual void GetChildObjects(uint32 InObjectHash, TArray<FTraceObjectInfo>& OutChildObjects) const = 0;
+
+	/** Set the filtered state for an individual object by its hash */
+	virtual void SetObjectFilterState(const FString& InObjectName, const bool bFilterState) = 0;
+	
+	/** Set timestamp for last processed update (data change) */
+	virtual const FDateTime& GetTimestamp() = 0;
+
+	/** Update filter preset */
+	virtual void UpdateFilterPreset(const TSharedPtr<ITraceFilterPreset> InPreset, bool IsEnabled) = 0;
+};
+
+} // namespace UE::TraceTools
\ No newline at end of file
diff --git a/Engine/Source/Developer/TraceTools/Private/Services/SessionTraceControllerFilterService.cpp b/Engine/Source/Developer/TraceTools/Private/Services/SessionTraceControllerFilterService.cpp
new file mode 100644
--- /dev/null
+++ b/Engine/Source/Developer/TraceTools/Private/Services/SessionTraceControllerFilterService.cpp
@@ -0,0 +1,118 @@
+// Copyright Epic Games, Inc. All Rights Reserved.
+
+#include "SessionTraceControllerFilterService.h"
+
+#include "Models/ITraceFilterPreset.h"
+#include "Misc/CoreDelegates.h"
+
+namespace UE::TraceTools
+{
+
+FSessionTraceControllerFilterService::FSessionTraceControllerFilterService(TSharedPtr<ITraceController> InTraceController)
+{
+	FCoreDelegates::OnEndFrame.AddRaw(this, &FSessionTraceControllerFilterService::OnApplyChannelChanges);
+
+	TraceController = InTraceController;
+	TraceController->OnStatusReceived().AddRaw(this, &FSessionTraceControllerFilterService::OnTraceStatusUpdated);
+}
+
+FSessionTraceControllerFilterService::~FSessionTraceControllerFilterService()
+{
+	FCoreDelegates::OnEndFrame.RemoveAll(this);
+
+	TraceController->OnStatusReceived().RemoveAll(this);
+}
+
+void FSessionTraceControllerFilterService::GetRootObjects(TArray<FTraceObjectInfo>& OutObjects) const
+{
+	OutObjects.Append(Objects);
+}
+
+void FSessionTraceControllerFilterService::GetChildObjects(uint32 InObjectHash, TArray<FTraceObjectInfo>& OutChildObjects) const
+{
+	/** TODO, parent/child relationship for Channels */
+}
+
+const FDateTime& FSessionTraceControllerFilterService::GetTimestamp()
+{
+	return TimeStamp;
+}
+
+void FSessionTraceControllerFilterService::SetObjectFilterState(const FString& InObjectName, const bool bFilterState)
+{
+	if (bFilterState)
+	{
+		FrameDisabledChannels.Remove(InObjectName);
+		FrameEnabledChannels.Add(InObjectName);
+	}
+	else
+	{
+		FrameEnabledChannels.Remove(InObjectName);
+		FrameDisabledChannels.Add(InObjectName);
+	}
+}
+
+void FSessionTraceControllerFilterService::UpdateFilterPreset(const TSharedPtr<ITraceFilterPreset> Preset, bool IsEnabled)
+{
+	TArray<FString> Names;
+	Preset->GetAllowlistedNames(Names);
+	if (IsEnabled)
+	{
+		FrameEnabledChannels.Append(Names);
+	}
+	else
+	{
+		FrameDisabledChannels.Append(Names);
+	}
+}
+
+void FSessionTraceControllerFilterService::DisableAllChannels()
+{
+	for (FTraceObjectInfo& ObjectInfo : Objects)
+	{
+		ObjectInfo.bEnabled = false;
+	}
+}
+
+void FSessionTraceControllerFilterService::OnTraceStatusUpdated(const FTraceStatus& InStatus, FTraceStatus::EUpdateType InUpdateType)
+{
+	if (!TraceController->HasAvailableSelectedInstance())
+	{
+		TimeStamp = FDateTime::Now();
+		Objects.Empty();
+		return;
+	}
+	if (InUpdateType != FTraceStatus::EUpdateType::ChannelsDesc && 
+		InUpdateType != FTraceStatus::EUpdateType::ChannelsStatus && 
+		InUpdateType != FTraceStatus::EUpdateType::All)
+	{
+		return;
+	}
+
+	TimeStamp = FDateTime::Now();
+
+	const TMap<uint32, FTraceStatus::FChannel> Channels = InStatus.Channels;
+	Objects.Empty(Channels.Num());
+
+	for (auto& Entry : Channels)
+	{
+		FTraceObjectInfo& EventInfo = Objects.AddDefaulted_GetRef();
+		EventInfo.Name = Entry.Value.Name;
+		EventInfo.bEnabled = Entry.Value.bEnabled;
+		EventInfo.bReadOnly = Entry.Value.bReadOnly;
+		EventInfo.Hash = Entry.Value.Id;
+		EventInfo.OwnerHash = 0;
+	}
+}
+
+void FSessionTraceControllerFilterService::OnApplyChannelChanges()
+{
+	if (FrameEnabledChannels.Num() || FrameDisabledChannels.Num())
+	{
+		TraceController->SetChannels(FrameEnabledChannels, FrameDisabledChannels);
+		FrameEnabledChannels.Empty();
+		FrameDisabledChannels.Empty();
+	}
+}
+
+} // namespace UE::TraceTools
\ No newline at end of file
diff --git a/Engine/Source/Developer/TraceTools/Private/Services/SessionTraceControllerFilterService.h b/Engine/Source/Developer/TraceTools/Private/Services/SessionTraceControllerFilterService.h
new file mode 100644
--- /dev/null
+++ b/Engine/Source/Developer/TraceTools/Private/Services/SessionTraceControllerFilterService.h
@@ -0,0 +1,59 @@
+// Copyright Epic Games, Inc. All Rights Reserved.
+
+#pragma once
+
+#include "ISessionTraceFilterService.h"
+#include "ITraceController.h"
+#include "Misc/DateTime.h"
+
+namespace TraceServices
+{
+	class IAnalysisSession;
+	typedef uint64 FSessionHandle;
+}
+
+namespace UE::TraceTools
+{
+
+class IEventInfoProvider;
+
+/** Implementation of ISessionTraceFilterService to query and set channels using the TraceController (MessageBus). */
+class FSessionTraceControllerFilterService : public ISessionTraceFilterService
+{
+public:
+	FSessionTraceControllerFilterService(TSharedPtr<ITraceController> InTraceController);
+	virtual ~FSessionTraceControllerFilterService();
+
+	/** Begin ISessionTraceFilterService overrides */
+	virtual void GetRootObjects(TArray<FTraceObjectInfo>& OutObjects) const override;
+	virtual void GetChildObjects(uint32 InObjectHash, TArray<FTraceObjectInfo>& OutChildObjects) const override;
+	virtual const FDateTime& GetTimestamp() override;
+	virtual void SetObjectFilterState(const FString& InObjectName, const bool bFilterState) override;
+	virtual void UpdateFilterPreset(const TSharedPtr<ITraceFilterPreset> InPreset, bool IsEnabled) override;
+	/** End ISessionTraceFilterService overrides */
+
+protected:
+	/** Callback at end of engine frame, used to dispatch all enabled/disabled channels */
+	void OnApplyChannelChanges();
+
+	/** Retrieves channels names from provider and marks them all as disabled */
+	void DisableAllChannels();
+
+	void RetrieveAndStoreStartupChannels();
+
+	void OnTraceStatusUpdated(const FTraceStatus& InStatus, FTraceStatus::EUpdateType InUpdateType);
+
+protected:
+	TSharedPtr<ITraceController> TraceController;
+
+	TArray<FTraceObjectInfo> Objects;
+
+	/** Names of channels that were either enabled or disabled during the duration of this frame */
+	TArray<FString> FrameEnabledChannels;
+	TArray<FString> FrameDisabledChannels;
+
+	/** Timestamp at which contained data (including provider) was last updated */
+	FDateTime TimeStamp;	
+};
+
+} // namespace UE::TraceTools
\ No newline at end of file
diff --git a/Engine/Source/Developer/TraceTools/Private/TraceToolsModule.cpp b/Engine/Source/Developer/TraceTools/Private/TraceToolsModule.cpp
new file mode 100644
--- /dev/null
+++ b/Engine/Source/Developer/TraceTools/Private/TraceToolsModule.cpp
@@ -0,0 +1,36 @@
+// Copyright Epic Games, Inc. All Rights Reserved.
+
+#include "TraceToolsModule.h"
+
+#include "Modules/ModuleManager.h"
+#include "Misc/ConfigContext.h"
+#include "TraceToolsStyle.h"
+#include "Widgets/Docking/SDockTab.h"
+
+// TraceTools
+#include "Widgets/STraceControl.h"
+
+IMPLEMENT_MODULE(FTraceToolsModule, TraceTools);
+
+FString FTraceToolsModule::TraceFiltersIni;
+
+void FTraceToolsModule::StartupModule()
+{
+	LLM_SCOPE_BYNAME(TEXT("Insights/TraceTools"));
+
+	UE::TraceTools::FTraceToolsStyle::Initialize();
+
+	FConfigContext::ReadIntoGConfig().Load(TEXT("Trace"), TraceFiltersIni);
+}
+
+void FTraceToolsModule::ShutdownModule()
+{
+	LLM_SCOPE_BYNAME(TEXT("Insights/TraceTools"));
+
+	UE::TraceTools::FTraceToolsStyle::Shutdown();
+}
+
+TSharedRef<SWidget> FTraceToolsModule::CreateTraceControlWidget(TSharedPtr<ITraceController> InTraceController)
+{
+	return SNew(UE::TraceTools::STraceControl, InTraceController);
+}
\ No newline at end of file
diff --git a/Engine/Source/Developer/TraceTools/Private/TraceToolsModule.h b/Engine/Source/Developer/TraceTools/Private/TraceToolsModule.h
new file mode 100644
--- /dev/null
+++ b/Engine/Source/Developer/TraceTools/Private/TraceToolsModule.h
@@ -0,0 +1,23 @@
+// Copyright Epic Games, Inc. All Rights Reserved.
+
+#pragma once
+
+#include "Interfaces/ITraceToolsModule.h"
+
+#include "Modules/ModuleInterface.h"
+
+class FName;
+class FString;
+class SWidget;
+
+class FTraceToolsModule : public UE::TraceTools::ITraceToolsModule
+{
+public:
+	// IModuleInterface interface
+	virtual void StartupModule() override;
+	virtual void ShutdownModule() override;
+
+	virtual TSharedRef<SWidget> CreateTraceControlWidget(TSharedPtr<ITraceController> InTraceController);
+
+	static FString TraceFiltersIni;
+};
\ No newline at end of file
diff --git a/Engine/Source/Developer/TraceTools/Private/TraceToolsStyle.cpp b/Engine/Source/Developer/TraceTools/Private/TraceToolsStyle.cpp
new file mode 100644
--- /dev/null
+++ b/Engine/Source/Developer/TraceTools/Private/TraceToolsStyle.cpp
@@ -0,0 +1,200 @@
+// Copyright Epic Games, Inc. All Rights Reserved.
+
+#include "TraceToolsStyle.h"
+
+#include "Brushes/SlateBoxBrush.h"
+#include "Brushes/SlateImageBrush.h"
+#include "Brushes/SlateNoResource.h"
+#include "Misc/Paths.h"
+#include "Styling/StyleColors.h"
+#include "Styling/SlateStyleRegistry.h"
+#include "Styling/SlateTypes.h"
+#include "Styling/CoreStyle.h"
+#include "Styling/SlateStyleMacros.h"
+
+namespace UE::TraceTools
+{
+
+// This is to fix the issue that SlateStyleMacros like IMAGE_BRUSH look for RootToContentDir but StyleSet->RootToContentDir is how this style is set up
+#define RootToContentDir StyleSet->RootToContentDir
+#define RootToCoreContentDir StyleSet->RootToCoreContentDir
+
+TSharedPtr< FSlateStyleSet > FTraceToolsStyle::StyleSet = nullptr;
+
+FTextBlockStyle FTraceToolsStyle::NormalText;
+
+#define ICON_FONT(...) FSlateFontInfo(RootToContentDir("Fonts/FontAwesome", TEXT(".ttf")), __VA_ARGS__)
+
+// Const icon sizes
+static const FVector2D Icon8x8(8.0f, 8.0f);
+static const FVector2D Icon9x19(9.0f, 19.0f);
+static const FVector2D Icon14x14(14.0f, 14.0f);
+static const FVector2D Icon16x16(16.0f, 16.0f);
+static const FVector2D Icon20x20(20.0f, 20.0f);
+static const FVector2D Icon22x22(22.0f, 22.0f);
+static const FVector2D Icon24x24(24.0f, 24.0f);
+static const FVector2D Icon28x28(28.0f, 28.0f);
+static const FVector2D Icon27x31(27.0f, 31.0f);
+static const FVector2D Icon26x26(26.0f, 26.0f);
+static const FVector2D Icon32x32(32.0f, 32.0f);
+static const FVector2D Icon40x40(40.0f, 40.0f);
+static const FVector2D Icon48x48(48.0f, 48.0f);
+static const FVector2D Icon75x82(75.0f, 82.0f);
+static const FVector2D Icon360x32(360.0f, 32.0f);
+static const FVector2D Icon171x39(171.0f, 39.0f);
+static const FVector2D Icon170x50(170.0f, 50.0f);
+static const FVector2D Icon267x140(170.0f, 50.0f);
+
+void FTraceToolsStyle::Initialize()
+{
+	LLM_SCOPE_BYNAME(TEXT("Insights/TraceTools"));
+
+	// Only register once
+	if( StyleSet.IsValid() )
+	{
+		return;
+	}
+
+	StyleSet = MakeShareable(new FSlateStyleSet("TraceToolsStyle") );
+	
+	StyleSet->SetContentRoot(FPaths::EngineContentDir() / TEXT("Editor/Slate"));
+	StyleSet->SetCoreContentRoot(FPaths::EngineContentDir() / TEXT("Slate"));
+
+	StyleSet->Set("EventFilter.GroupBorder", new BOX_BRUSH("Common/GroupBorder", FMargin(4.0f / 16.0f)));
+
+	NormalText = FTextBlockStyle()
+		.SetFont(DEFAULT_FONT("Regular", FCoreStyle::RegularTextSize))
+		.SetColorAndOpacity(FSlateColor::UseForeground())
+		.SetShadowOffset(FVector2D::ZeroVector)
+		.SetShadowColorAndOpacity(FLinearColor::Black)
+		.SetHighlightColor(FLinearColor(0.02f, 0.3f, 0.0f))
+		.SetHighlightShape(BOX_BRUSH("Common/TextBlockHighlightShape", FMargin(3.f / 8.f)));
+
+	// Colors
+	{
+		StyleSet->Set("EventFilter.EnginePreset", FLinearColor(0.728f, 0.364f, 0.003f));
+		StyleSet->Set("EventFilter.SharedPreset", FLinearColor(0.003f, 0.364f, 0.728f));
+		StyleSet->Set("EventFilter.LocalPreset", FLinearColor(0.003f, 0.728f, 0.364f));
+	}
+	
+	// Icons
+	{
+		StyleSet->Set("EventFilter.State.Enabled", new IMAGE_BRUSH("Common/CheckBox_Checked", Icon16x16));
+		StyleSet->Set("EventFilter.State.Enabled_Hovered", new IMAGE_BRUSH("Common/CheckBox_Checked_Hovered", Icon16x16));
+
+		StyleSet->Set("EventFilter.State.Disabled", new IMAGE_BRUSH("Common/CheckBox", Icon16x16));
+		StyleSet->Set("EventFilter.State.Disabled_Hovered", new IMAGE_BRUSH("Common/CheckBox_Hovered", Icon16x16));
+
+		StyleSet->Set("EventFilter.State.Pending", new IMAGE_BRUSH("Common/CheckBox_Undetermined", Icon16x16));
+		StyleSet->Set("EventFilter.State.Pending_Hovered", new IMAGE_BRUSH("Common/CheckBox_Undetermined_Hovered", Icon16x16));
+
+		StyleSet->Set("EventFilter.TabIcon", new IMAGE_BRUSH_SVG("Starship/Common/TraceDataFiltering", Icon16x16));
+	}
+
+	FButtonStyle Button = FButtonStyle()
+		.SetNormal(BOX_BRUSH("Common/Button", FVector2D(32, 32), 8.0f / 32.0f))
+		.SetHovered(BOX_BRUSH("Common/Button_Hovered", FVector2D(32, 32), 8.0f / 32.0f))
+		.SetPressed(BOX_BRUSH("Common/Button_Pressed", FVector2D(32, 32), 8.0f / 32.0f))
+		.SetNormalPadding(FMargin(2, 2, 2, 2))
+		.SetPressedPadding(FMargin(2, 3, 2, 1));
+
+	StyleSet->Set("ToggleButton", FButtonStyle(Button)
+		.SetNormal(FSlateNoResource())
+		.SetHovered(BOX_BRUSH("Common/RoundedSelection_16x", 4.0f / 16.0f, FLinearColor(0.701f, 0.225f, 0.003f)))
+		.SetPressed(BOX_BRUSH("Common/RoundedSelection_16x", 4.0f / 16.0f, FLinearColor(0.701f, 0.225f, 0.003f)))
+	);
+
+	FComboButtonStyle ToolbarComboButton = FComboButtonStyle()
+		.SetButtonStyle(StyleSet->GetWidgetStyle<FButtonStyle>("ToggleButton"))
+		.SetDownArrowImage(IMAGE_BRUSH("Common/ShadowComboArrow", Icon8x8))
+		.SetMenuBorderBrush(BOX_BRUSH("Old/Menu_Background", FMargin(8.0f / 64.0f)))
+		.SetMenuBorderPadding(FMargin(0.0f));
+	StyleSet->Set("EventFilter.ComboButton", ToolbarComboButton);
+	
+	StyleSet->Set("EventFilter.TextStyle", FTextBlockStyle(NormalText)
+		.SetFont(DEFAULT_FONT("Bold", 9))
+		.SetColorAndOpacity(FLinearColor(1.0f, 1.0f, 1.0f, 0.9f))
+		.SetShadowOffset(FVector2D(1, 1))
+		.SetShadowColorAndOpacity(FLinearColor(0, 0, 0, 0.9f)));
+	
+
+	// Filter list
+	/* Set images for various SCheckBox states associated with "ContentBrowser.FilterButton" ... */
+	const FCheckBoxStyle FilterButtonCheckBoxStyle = FCheckBoxStyle()
+		.SetUncheckedImage(IMAGE_BRUSH("ContentBrowser/FilterUnchecked", FVector2D(10.0f, 20.0f)))
+		.SetUncheckedHoveredImage(IMAGE_BRUSH("ContentBrowser/FilterUnchecked", FVector2D(10.0f, 20.0f), FLinearColor(0.5f, 0.5f, 0.5f, 1.0f)))
+		.SetUncheckedPressedImage(IMAGE_BRUSH("ContentBrowser/FilterUnchecked", FVector2D(10.0f, 20.0f), FLinearColor(0.5f, 0.5f, 0.5f, 1.0f)))
+		.SetCheckedImage(IMAGE_BRUSH("ContentBrowser/FilterChecked", FVector2D(10.0f, 20.0f)))
+		.SetCheckedHoveredImage(IMAGE_BRUSH("ContentBrowser/FilterChecked", FVector2D(10.0f, 20.0f), FLinearColor(0.5f, 0.5f, 0.5f, 1.0f)))
+		.SetCheckedPressedImage(IMAGE_BRUSH("ContentBrowser/FilterChecked", FVector2D(10.0f, 20.0f), FLinearColor(0.5f, 0.5f, 0.5f, 1.0f)));
+	/* ... and add the new style */
+	StyleSet->Set("FilterPresets.FilterButton", FilterButtonCheckBoxStyle);
+
+	StyleSet->Set("FilterPresets.FilterNameFont", DEFAULT_FONT("Regular", 10));
+	StyleSet->Set("FilterPresets.FilterButtonBorder", new BOX_BRUSH("Common/RoundedSelection_16x", FMargin(4.0f / 16.0f)));
+
+	StyleSet->Set("FilterPresets.SessionWarningBorder", new BOX_BRUSH("Common/DarkGroupBorder", FMargin(4.0f / 16.0f)));
+	StyleSet->Set("FilterPresets.WarningIcon", new CORE_IMAGE_BRUSH_SVG("Starship/Common/alert-circle", Icon40x40, FStyleColors::Warning));
+
+	StyleSet->Set("FontAwesome.9", ICON_FONT(9));
+
+	StyleSet->Set("TraceControl.StartTrace", new CORE_IMAGE_BRUSH_SVG("Starship/Insights/TraceTools/TraceStart", Icon40x40));
+	StyleSet->Set("TraceControl.StartTrace.Small", new CORE_IMAGE_BRUSH_SVG("Starship/Insights/TraceTools/TraceStart", Icon20x20));
+
+	StyleSet->Set("TraceControl.StopTrace", new CORE_IMAGE_BRUSH_SVG("Starship/Insights/TraceTools/TraceStop", Icon40x40));
+	StyleSet->Set("TraceControl.StopTrace.Small", new CORE_IMAGE_BRUSH_SVG("Starship/Insights/TraceTools/TraceStop", Icon20x20));
+
+	StyleSet->Set("TraceControl.PauseTrace", new CORE_IMAGE_BRUSH_SVG("Starship/Insights/TraceTools/TracePause", Icon40x40));
+	StyleSet->Set("TraceControl.PauseTrace.Small", new CORE_IMAGE_BRUSH_SVG("Starship/Insights/TraceTools/TracePause", Icon20x20));
+
+	StyleSet->Set("TraceControl.ResumeTrace", new CORE_IMAGE_BRUSH_SVG("Starship/Insights/TraceTools/TraceResume", Icon40x40));
+	StyleSet->Set("TraceControl.ResumeTrace.Small", new CORE_IMAGE_BRUSH_SVG("Starship/Insights/TraceTools/TraceResume", Icon20x20));
+
+	StyleSet->Set("TraceControl.TraceSnapshot", new CORE_IMAGE_BRUSH_SVG("Starship/Insights/TraceTools/TraceSnapshot", Icon40x40));
+	StyleSet->Set("TraceControl.TraceSnapshot.Small", new CORE_IMAGE_BRUSH_SVG("Starship/Insights/TraceTools/TraceSnapshot", Icon20x20));
+
+	StyleSet->Set("TraceControl.SetTraceTargetServer", new CORE_IMAGE_BRUSH_SVG("Starship/Common/server", Icon40x40));
+	StyleSet->Set("TraceControl.SetTraceTargetServer.Small", new CORE_IMAGE_BRUSH_SVG("Starship/Common/server", Icon20x20));
+
+	StyleSet->Set("TraceControl.SetTraceTargetFile", new CORE_IMAGE_BRUSH_SVG("Starship/Common/file", Icon40x40));
+	StyleSet->Set("TraceControl.SetTraceTargetFile.Small", new CORE_IMAGE_BRUSH_SVG("Starship/Common/file", Icon20x20));
+
+	StyleSet->Set("TraceControl.TraceScreenshot", new CORE_IMAGE_BRUSH_SVG("../Editor/Slate/Starship/Common/HighResolutionScreenshot", Icon40x40));
+	StyleSet->Set("TraceControl.TraceScreenshot.Small", new CORE_IMAGE_BRUSH_SVG("../Editor/Slate/Starship/Common/HighResolutionScreenshot", Icon20x20));
+
+	StyleSet->Set("TraceControl.TraceBookmark", new CORE_IMAGE_BRUSH_SVG("../Editor/Slate/Starship/Common/Bookmarks", Icon40x40));
+	StyleSet->Set("TraceControl.TraceBookmark.Small", new CORE_IMAGE_BRUSH_SVG("../Editor/Slate/Starship/Common/Bookmarks", Icon20x20));
+
+	StyleSet->Set("TraceControl.ToggleStatNamedEvents", new CORE_IMAGE_BRUSH_SVG("../Editor/Slate/Starship/Common/Statistics", Icon40x40));
+	StyleSet->Set("TraceControl.ToggleStatNamedEvents.Small", new CORE_IMAGE_BRUSH_SVG("../Editor/Slate/Starship/Common/Statistics", Icon20x20));
+	
+	FSlateStyleRegistry::RegisterSlateStyle( *StyleSet.Get() );
+}
+
+#undef IMAGE_BRUSH
+#undef BOX_BRUSH
+#undef BORDER_BRUSH
+#undef DEFAULT_FONT
+#undef ICON_FONT
+
+void FTraceToolsStyle::Shutdown()
+{
+	if( StyleSet.IsValid() )
+	{
+		FSlateStyleRegistry::UnRegisterSlateStyle( *StyleSet.Get() );
+		ensure( StyleSet.IsUnique() );
+		StyleSet.Reset();
+	}
+}
+
+const ISlateStyle& FTraceToolsStyle::Get()
+{
+	return *( StyleSet.Get() );
+}
+
+const FName& FTraceToolsStyle::GetStyleSetName()
+{
+	return StyleSet->GetStyleSetName();
+}
+
+} // namespace UE::TraceTools
\ No newline at end of file
diff --git a/Engine/Source/Developer/TraceTools/Private/TraceToolsStyle.h b/Engine/Source/Developer/TraceTools/Private/TraceToolsStyle.h
new file mode 100644
--- /dev/null
+++ b/Engine/Source/Developer/TraceTools/Private/TraceToolsStyle.h
@@ -0,0 +1,33 @@
+// Copyright Epic Games, Inc. All Rights Reserved.
+
+#pragma once
+
+#include "Styling/SlateStyle.h"
+
+class ISlateStyle;
+struct FTextBlockStyle;
+
+namespace UE::TraceTools
+{
+
+class FTraceToolsStyle
+{
+public:
+	static void Initialize();
+	static void Shutdown();
+
+	static const ISlateStyle& Get();
+	static const FName& GetStyleSetName();
+	static const FSlateBrush* GetBrush(FName PropertyName, const ANSICHAR* Specifier = NULL)
+	{
+		return StyleSet->GetBrush(PropertyName, Specifier);
+	}
+
+private:
+
+	/** Singleton instances of this style. */
+	static TSharedPtr< class FSlateStyleSet > StyleSet;	
+	static FTextBlockStyle NormalText;
+};
+
+} // namespace UE::TraceTools
diff --git a/Engine/Source/Developer/TraceTools/Private/Widgets/SFilterPreset.cpp b/Engine/Source/Developer/TraceTools/Private/Widgets/SFilterPreset.cpp
new file mode 100644
--- /dev/null
+++ b/Engine/Source/Developer/TraceTools/Private/Widgets/SFilterPreset.cpp
@@ -0,0 +1,324 @@
+// Copyright Epic Games, Inc. All Rights Reserved.
+
+#include "SFilterPreset.h"
+
+#include "Framework/MultiBox/MultiBoxBuilder.h"
+#include "Widgets/Layout/SBorder.h"
+#include "Widgets/Text/STextBlock.h"
+
+// TraceTools
+#include "Models/ITraceFilterPreset.h"
+#include "TraceToolsStyle.h"
+
+#define LOCTEXT_NAMESPACE "SFilterPreset"
+
+namespace UE::TraceTools
+{
+
+SFilterPreset::~SFilterPreset()
+{
+	if (IsHovered())
+	{
+		OnHighlightPreset.ExecuteIfBound(nullptr);
+	}
+}
+
+void SFilterPreset::Construct(const FArguments& InArgs)
+{
+	bEnabled = false;
+	bHighlighted = false;
+	OnPresetChanged = InArgs._OnPresetChanged;
+	OnRequestRemove = InArgs._OnRequestRemove;
+	OnRequestEnableOnly = InArgs._OnRequestEnableOnly;
+	OnRequestEnableAll = InArgs._OnRequestEnableAll;
+	OnRequestDisableAll = InArgs._OnRequestDisableAll;
+	OnRequestRemoveAll = InArgs._OnRequestRemoveAll;
+	OnRequestDelete = InArgs._OnRequestDelete;
+	OnRequestSave = InArgs._OnRequestSave;
+	OnHighlightPreset = InArgs._OnHighlightPreset;
+
+	FilterPreset = InArgs._FilterPreset;
+
+	const FName ColorName = [this]()
+	{
+		if (FilterPreset->CanDelete())
+		{
+			// User preset 
+			if (FilterPreset->IsLocal())
+			{
+				return FName("EventFilter.LocalPreset");
+			}
+			
+			return FName("EventFilter.SharedPreset");
+		}
+		
+		// Engine preset
+		return FName("EventFilter.EnginePreset");
+	}();
+
+	FilterColor = FTraceToolsStyle::Get().GetColor(ColorName);
+
+	ChildSlot
+	[
+		SNew(SBorder)
+		.Padding(0)
+		.BorderBackgroundColor(FLinearColor(0.2f, 0.2f, 0.2f, 0.2f))
+		.BorderImage(FTraceToolsStyle::GetBrush("FilterPresets.FilterButtonBorder"))
+		[
+			SAssignNew(ToggleButtonPtr, SFilterPresetCheckBox)
+			.Style(FTraceToolsStyle::Get(), "FilterPresets.FilterButton")
+			.Padding(this, &SFilterPreset::GetNameLabelPadding)
+			.IsChecked(this, &SFilterPreset::IsChecked)
+			.OnCheckStateChanged(this, &SFilterPreset::PresetToggled)
+			.OnGetMenuContent(this, &SFilterPreset::GetRightClickMenuContent)
+			.ForegroundColor(this, &SFilterPreset::GetPresetForegroundColor)
+			[
+				SNew(STextBlock)
+				.ColorAndOpacity(this, &SFilterPreset::GetNameLabelColorAndOpacity)
+				.Font(FTraceToolsStyle::Get().GetFontStyle("FilterPresets.FilterNameFont"))
+				.ShadowOffset(FVector2D(1.f, 1.f))
+				.Text(this, &SFilterPreset::GetPresetName)
+			]
+		]
+	];
+
+	ToggleButtonPtr->SetOnFilterCtrlClicked(FOnClicked::CreateSP(this, &SFilterPreset::FilterCtrlClicked));
+	ToggleButtonPtr->SetOnFilterAltClicked(FOnClicked::CreateSP(this, &SFilterPreset::FilterAltClicked));
+	ToggleButtonPtr->SetOnFilterDoubleClicked(FOnClicked::CreateSP(this, &SFilterPreset::FilterDoubleClicked));
+	ToggleButtonPtr->SetOnFilterMiddleButtonClicked(FOnClicked::CreateSP(this, &SFilterPreset::FilterMiddleButtonClicked));
+
+	TAttribute<FText> Attribute = TAttribute<FText>::Create(TAttribute<FText>::FGetter::CreateSP(this, &SFilterPreset::GetToolTipText));
+	ToggleButtonPtr->SetToolTipText(Attribute);
+}
+
+void SFilterPreset::SetEnabled(bool InEnabled)
+{
+	if (InEnabled != bEnabled)
+	{
+		bEnabled = InEnabled;
+ 		OnPresetChanged.ExecuteIfBound(*this);
+	}
+}
+
+bool SFilterPreset::IsEnabled() const
+{
+	return bEnabled;
+}
+
+const TSharedPtr<ITraceFilterPreset>& SFilterPreset::GetFilterPreset() const
+{
+	return FilterPreset;
+}
+
+void SFilterPreset::OnMouseLeave(const FPointerEvent& MouseEvent)
+{
+	SCompoundWidget::OnMouseLeave(MouseEvent);
+	
+	if (bHighlighted)
+	{
+		OnHighlightPreset.ExecuteIfBound(nullptr);
+	}
+	
+	bHighlighted = false;
+}
+
+FReply SFilterPreset::OnMouseMove(const FGeometry& MyGeometry, const FPointerEvent& MouseEvent)
+{
+	if (IsHovered())
+	{
+		const bool bWasHighlighted = bHighlighted;
+		bHighlighted = MouseEvent.IsControlDown();
+		if (bHighlighted || bWasHighlighted)
+		{
+			OnHighlightPreset.ExecuteIfBound(bHighlighted ? FilterPreset : nullptr);
+		}
+	}
+
+	return SCompoundWidget::OnMouseMove(MyGeometry, MouseEvent);
+}
+
+void SFilterPreset::PresetToggled(ECheckBoxState NewState)
+{
+	SetEnabled(NewState == ECheckBoxState::Checked);
+}
+
+FReply SFilterPreset::FilterCtrlClicked()
+{
+	OnRequestEnableAll.ExecuteIfBound();
+	return FReply::Handled();
+}
+
+FReply SFilterPreset::FilterAltClicked()
+{
+	OnRequestDisableAll.ExecuteIfBound();
+	return FReply::Handled();
+}
+
+FReply SFilterPreset::FilterDoubleClicked()
+{
+	// Disable all other presets and enable this one.
+	OnRequestDisableAll.ExecuteIfBound();
+	SetEnabled(true);
+	return FReply::Handled();
+}
+
+FReply SFilterPreset::FilterMiddleButtonClicked()
+{
+	RemovePreset();
+	return FReply::Handled();
+}
+
+FText SFilterPreset::GetToolTipText() const
+{
+	return bHighlighted ? FilterPreset->GetDescription() : FText::Format(LOCTEXT("FilterPresetTooltipText", "{0}\nHold CTRL to highlight allowlisted items."), FilterPreset->GetDescription());
+}
+
+TSharedRef<SWidget> SFilterPreset::GetRightClickMenuContent()
+{
+	FMenuBuilder MenuBuilder(/*bInShouldCloseWindowAfterMenuSelection=*/true, NULL);
+
+	MenuBuilder.BeginSection("PresetOptions", LOCTEXT("FilterContextHeading", "Preset Options"));
+	{
+		MenuBuilder.AddMenuEntry(
+			FText::Format(LOCTEXT("RemovePreset", "Remove: {0}"), GetPresetName()),
+			LOCTEXT("RemovePresetTooltip", "Remove this preset from the list. It can be added again in the presets menu."),
+			FSlateIcon(),
+			FUIAction(FExecuteAction::CreateSP(this, &SFilterPreset::RemovePreset))
+		);
+
+		MenuBuilder.AddMenuEntry(
+			FText::Format(LOCTEXT("EnableOnlyThisPreset", "Enable this only: {0}"), GetPresetName()),
+			LOCTEXT("EnableOnlyThisPresetTooltip", "Enable only this prest from the list."),
+			FSlateIcon(),
+			FUIAction(FExecuteAction::CreateSP(this, &SFilterPreset::EnableOnly))
+		);
+
+	}
+	MenuBuilder.EndSection();
+
+	MenuBuilder.BeginSection("PresetBulkOptions", LOCTEXT("BulkFilterContextHeading", "Bulk Preset Options"));
+	{
+		MenuBuilder.AddMenuEntry(
+			LOCTEXT("EnableAllPresets", "Enable All Presets"),
+			LOCTEXT("EnableAllPresetsTooltip", "Enables all presets."),
+			FSlateIcon(),
+			FUIAction(FExecuteAction::CreateSP(this, &SFilterPreset::EnableAllPresets))
+		);
+
+		MenuBuilder.AddMenuEntry(
+			LOCTEXT("DisableAllPresets", "Disable All Presets"),
+			LOCTEXT("DisableAllPresetsTooltip", "Disables all active presets."),
+			FSlateIcon(),
+			FUIAction(FExecuteAction::CreateSP(this, &SFilterPreset::DisableAllPresets))
+		);
+
+		MenuBuilder.AddMenuEntry(
+			LOCTEXT("RemoveAllPresets", "Remove All Presets"),
+			LOCTEXT("RemoveAllPresetsTooltip", "Removes all presets from the list."),
+			FSlateIcon(),
+			FUIAction(FExecuteAction::CreateSP(this, &SFilterPreset::RemoveAllPresets))
+		);
+	}
+	MenuBuilder.EndSection();
+
+
+	MenuBuilder.BeginSection("UserPresetOptions", LOCTEXT("UserPresetsContextHeading", "User Presets"));
+	{
+		if (FilterPreset->CanDelete())
+		{
+			MenuBuilder.AddMenuEntry(
+				FText::Format(LOCTEXT("DeletePreset", "Delete User Preset: {0}"), GetPresetName()),
+				LOCTEXT("DeleteFilterTooltip", "Deletes this User Preset."),
+				FSlateIcon(),
+				FUIAction(FExecuteAction::CreateSP(this, &SFilterPreset::DeletePreset))
+			);
+
+			MenuBuilder.AddMenuEntry(
+				FText::Format(LOCTEXT("SavePreset", "Save User Preset: {0}"), GetPresetName()),
+				LOCTEXT("SaveFilterTooltip", "Saves the currently filtering state as the selected User Preset."),
+				FSlateIcon(),
+				FUIAction(FExecuteAction::CreateSP(this, &SFilterPreset::SavePreset))
+			);
+		}
+	}
+	MenuBuilder.EndSection();
+
+
+	return MenuBuilder.MakeWidget();
+}
+
+void SFilterPreset::RemovePreset()
+{
+	TSharedRef<SFilterPreset> Self = SharedThis(this);
+	OnRequestRemove.ExecuteIfBound(Self);
+}
+
+void SFilterPreset::EnableOnly()
+{
+	TSharedRef<SFilterPreset> Self = SharedThis(this);
+	OnRequestEnableOnly.ExecuteIfBound(Self);
+}
+
+void SFilterPreset::EnableAllPresets()
+{
+	OnRequestEnableAll.ExecuteIfBound();
+}
+
+void SFilterPreset::DisableAllPresets()
+{
+	OnRequestDisableAll.ExecuteIfBound();
+}
+
+void SFilterPreset::RemoveAllPresets()
+{
+	OnRequestRemoveAll.ExecuteIfBound();
+}
+
+void SFilterPreset::SavePreset()
+{
+	TSharedRef<SFilterPreset> Self = SharedThis(this);
+	OnRequestSave.ExecuteIfBound(Self);
+}
+
+void SFilterPreset::DeletePreset()
+{
+	TSharedRef<SFilterPreset> Self = SharedThis(this);
+	OnRequestDelete.ExecuteIfBound(Self);
+}
+
+ECheckBoxState SFilterPreset::IsChecked() const
+{
+	return bEnabled ? ECheckBoxState::Checked : ECheckBoxState::Unchecked;
+}
+
+FSlateColor SFilterPreset::GetPresetForegroundColor() const
+{
+	return IsChecked() == ECheckBoxState::Checked ? FilterColor : FLinearColor::White;
+}
+
+FMargin SFilterPreset::GetNameLabelPadding() const
+{
+	return ToggleButtonPtr->IsPressed() ? FMargin(3, 2, 4, 0) : FMargin(3, 1, 4, 1);
+}
+
+FSlateColor SFilterPreset::GetNameLabelColorAndOpacity() const
+{
+	const float DimFactor = 0.75f;
+	return IsHovered() ? FLinearColor(DimFactor, DimFactor, DimFactor, 1.0f) : FLinearColor::White;
+}
+
+FText SFilterPreset::GetPresetName() const
+{
+	FText PresetDisplayText = FilterPreset->GetDisplayText();
+
+	if (PresetDisplayText.IsEmpty())
+	{
+		PresetDisplayText = LOCTEXT("UnknownPresetName", "???");
+	}
+
+	return PresetDisplayText;
+}
+
+} // namespace UE::TraceTools
+
+#undef LOCTEXT_NAMESPACE
diff --git a/Engine/Source/Developer/TraceTools/Private/Widgets/SFilterPreset.h b/Engine/Source/Developer/TraceTools/Private/Widgets/SFilterPreset.h
new file mode 100644
--- /dev/null
+++ b/Engine/Source/Developer/TraceTools/Private/Widgets/SFilterPreset.h
@@ -0,0 +1,231 @@
+// Copyright Epic Games, Inc. All Rights Reserved.
+
+#pragma once
+
+#include "Widgets/Input/SCheckBox.h"
+
+enum class ECheckBoxState : uint8;
+
+namespace UE::TraceTools
+{
+
+struct ITraceFilterPreset;
+
+DECLARE_DELEGATE_OneParam(FOnPresetChanged, const class SFilterPreset&);
+
+/** A class for check boxes in the preset list. If you double click the checkbox, you will enable it and disable all others */
+class SFilterPresetCheckBox : public SCheckBox
+{
+public:
+	
+	void SetOnFilterCtrlClicked(const FOnClicked& NewFilterCtrlClicked)
+	{
+		OnFilterCtrlClicked = NewFilterCtrlClicked;
+	}
+
+	void SetOnFilterAltClicked(const FOnClicked& NewFilteAltClicked)
+	{
+		OnFilterAltClicked = NewFilteAltClicked;
+	}
+
+	void SetOnFilterDoubleClicked(const FOnClicked& NewFilterDoubleClicked)
+	{
+		OnFilterDoubleClicked = NewFilterDoubleClicked;
+	}
+
+	void SetOnFilterMiddleButtonClicked(const FOnClicked& NewFilterMiddleButtonClicked)
+	{
+		OnFilterMiddleButtonClicked = NewFilterMiddleButtonClicked;
+	}
+
+	virtual FReply OnMouseButtonDoubleClick(const FGeometry& InMyGeometry, const FPointerEvent& InMouseEvent) override
+	{
+		// Handle left mouse button double click
+		if (InMouseEvent.GetEffectingButton() == EKeys::LeftMouseButton && OnFilterDoubleClicked.IsBound())
+		{
+			return OnFilterDoubleClicked.Execute();
+		}
+		else
+		{
+			return SCheckBox::OnMouseButtonDoubleClick(InMyGeometry, InMouseEvent);
+		}
+	}
+
+	virtual FReply OnMouseButtonUp(const FGeometry& InMyGeometry, const FPointerEvent& InMouseEvent) override
+	{
+		SCheckBox::OnMouseButtonUp(InMyGeometry, InMouseEvent);
+
+		if (InMouseEvent.IsControlDown() && OnFilterCtrlClicked.IsBound())
+		{
+			return OnFilterCtrlClicked.Execute();
+		}
+		else if (InMouseEvent.IsAltDown() && OnFilterAltClicked.IsBound())
+		{
+			return OnFilterAltClicked.Execute();
+		}
+		else if (InMouseEvent.GetEffectingButton() == EKeys::MiddleMouseButton && OnFilterMiddleButtonClicked.IsBound())
+		{
+			return OnFilterMiddleButtonClicked.Execute();
+		}
+		else
+		{			
+			return FReply::Handled().ReleaseMouseCapture();
+		}
+	}
+
+private:
+	FOnClicked OnFilterCtrlClicked;
+	FOnClicked OnFilterAltClicked;
+	FOnClicked OnFilterDoubleClicked;
+	FOnClicked OnFilterMiddleButtonClicked;
+};
+
+/**
+ * A single filtering preset in the preset list. Can be removed by clicking the remove button on it.
+ */
+class SFilterPreset : public SCompoundWidget
+{
+public:
+	DECLARE_DELEGATE_OneParam(FOnRequestRemove, const TSharedRef<SFilterPreset>& /*PresetToRemove*/);
+	DECLARE_DELEGATE_OneParam(FOnRequestEnableOnly, const TSharedRef<SFilterPreset>& /*PresetToEnable*/);
+	DECLARE_DELEGATE_OneParam(FOnRequestDelete, const TSharedRef<SFilterPreset>& /*PresetToDelete*/);
+	DECLARE_DELEGATE_OneParam(FOnRequestSave, const TSharedRef<SFilterPreset>& /*PresetToSave*/);
+	DECLARE_DELEGATE_OneParam(FOnHighlightPreset, const TSharedPtr<ITraceFilterPreset>& /*PresetToHighlight*/);
+	DECLARE_DELEGATE(FOnRequestEnableAll);
+	DECLARE_DELEGATE(FOnRequestDisableAll);
+	DECLARE_DELEGATE(FOnRequestRemoveAll);	
+
+	SLATE_BEGIN_ARGS(SFilterPreset) {}
+		/** Filter preset this widget represents */
+		SLATE_ARGUMENT(TSharedPtr<ITraceFilterPreset>, FilterPreset)
+
+		/** Invoked when the preset toggled */
+		SLATE_EVENT(FOnPresetChanged, OnPresetChanged)
+
+		/** Invoked when a request to remove this preset originated from within this widget */
+		SLATE_EVENT(FOnRequestRemove, OnRequestRemove)
+
+		/** Invoked when a request to enable only this preset originated from within this widget */
+		SLATE_EVENT(FOnRequestEnableOnly, OnRequestEnableOnly)
+
+		/** Invoked when a request to delete this preset originated from within this widget */
+		SLATE_EVENT(FOnRequestDelete, OnRequestDelete);
+
+		/** Invoked when a request to save the current filtering state as this preset originated from within this widget */
+		SLATE_EVENT(FOnRequestSave, OnRequestSave);
+
+		/** Invoked when a request to highlight this preset originated from within this widget */
+		SLATE_EVENT(FOnHighlightPreset, OnHighlightPreset)
+
+		/** Invoked when a request to enable all presets originated from within this widget */
+		SLATE_EVENT(FOnRequestEnableAll, OnRequestEnableAll)
+
+		/** Invoked when a request to disable all presets originated from within this widget */
+		SLATE_EVENT(FOnRequestDisableAll, OnRequestDisableAll)
+
+		/** Invoked when a request to remove all presets originated from within this widget */
+		SLATE_EVENT(FOnRequestRemoveAll, OnRequestRemoveAll)
+	SLATE_END_ARGS()
+
+	virtual ~SFilterPreset();
+
+	/** Constructs this widget with InArgs */
+	void Construct(const FArguments& InArgs);
+
+	/** Begin SCompoundWidget overrides */
+	virtual void OnMouseLeave(const FPointerEvent& MouseEvent) override;
+	virtual FReply OnMouseMove(const FGeometry& MyGeometry, const FPointerEvent& MouseEvent) override;
+	/** End SCompoundWidget overrides */
+
+	/** Sets whether or not this preset is applied to the combined filtering state */
+	void SetEnabled(bool InEnabled);
+
+	/** Returns true if this preset contributes to the combined filtering state */
+	bool IsEnabled() const;
+
+	/** Get the preset object this widget represents */
+	const TSharedPtr<ITraceFilterPreset>& GetFilterPreset() const;
+private:
+	/** Handler for when the preset checkbox is clicked */
+	void PresetToggled(ECheckBoxState NewState);
+
+	/** Handler for when the preset checkbox is clicked and a control key is pressed */
+	FReply FilterCtrlClicked();
+
+	/** Handler for when the preset checkbox is clicked and an alt key is pressed */
+	FReply FilterAltClicked();
+
+	/** Handler for when the preset checkbox is double clicked */
+	FReply FilterDoubleClicked();
+
+	/** Handler for when the preset checkbox is middle button clicked */
+	FReply FilterMiddleButtonClicked();
+
+	FText GetToolTipText() const;
+
+	/** Handler to create a right click menu */
+	TSharedRef<SWidget> GetRightClickMenuContent();
+
+	/** Removes this preset from the presets list */
+	void RemovePreset();
+
+	/** Enables only this preset from the presets list */
+	void EnableOnly();
+
+	/** Enables all presets in the list */
+	void EnableAllPresets();
+
+	/** Disables all active presets in the list */
+	void DisableAllPresets();
+
+	/** Removes all presets in the list */
+	void RemoveAllPresets();
+
+	/** Saves the current filtering state as this preset */
+	void SavePreset();
+
+	/** Deletes this preset */
+	void DeletePreset();
+
+	/** Handler to determine the "checked" state of the preset checkbox */
+	ECheckBoxState IsChecked() const;
+
+	/** Handler to determine the color of the checkbox when it is checked */
+	FSlateColor GetPresetForegroundColor() const;
+
+	/** Handler to determine the padding of the checkbox text when it is pressed */
+	FMargin GetNameLabelPadding() const;
+
+	/** Handler to determine the color of the checkbox text when it is hovered */
+	FSlateColor GetNameLabelColorAndOpacity() const;
+
+	/** Returns the display name for this filter */
+	FText GetPresetName() const;
+
+private:
+	TSharedPtr<ITraceFilterPreset> FilterPreset;
+
+	FOnPresetChanged OnPresetChanged;
+	FOnRequestRemove OnRequestRemove;
+	FOnRequestDelete OnRequestDelete;
+	FOnRequestSave OnRequestSave;
+	FOnHighlightPreset OnHighlightPreset;
+	FOnRequestEnableOnly OnRequestEnableOnly;
+	FOnRequestEnableAll OnRequestEnableAll;
+	FOnRequestDisableAll OnRequestDisableAll;
+	FOnRequestDisableAll OnRequestRemoveAll;	
+
+	/** True when this preset should be applied to filtering state */
+	bool bEnabled;
+
+	/** The button to toggle the preset on or off */
+	TSharedPtr<SFilterPresetCheckBox> ToggleButtonPtr;
+
+	/** The color of the checkbox for this preset */
+	FLinearColor FilterColor;
+
+	/** True when this preset its filtering state is currently set to being highlighted */
+	bool bHighlighted;
+};
+
+} // namespace UE::TraceTools
\ No newline at end of file
diff --git a/Engine/Source/Developer/TraceTools/Private/Widgets/SFilterPresetList.cpp b/Engine/Source/Developer/TraceTools/Private/Widgets/SFilterPresetList.cpp
new file mode 100644
--- /dev/null
+++ b/Engine/Source/Developer/TraceTools/Private/Widgets/SFilterPresetList.cpp
@@ -0,0 +1,539 @@
+// Copyright Epic Games, Inc. All Rights Reserved.
+
+#include "SFilterPresetList.h"
+
+#include "Framework/Application/SlateApplication.h"
+#include "Framework/MultiBox/MultiBoxBuilder.h"
+#include "Layout/WidgetPath.h"
+#include "Misc/ConfigCacheIni.h"
+#include "Widgets/Layout/SWrapBox.h"
+#include "Widgets/Input/SEditableTextBox.h"
+
+// TraceTools
+#include "Models/TraceFilterPresets.h"
+#include "TraceToolsModule.h"
+
+#if WITH_EDITOR
+#include "SSettingsEditorCheckoutNotice.h"
+#endif
+
+#define LOCTEXT_NAMESPACE "InsightsFilterList"
+
+namespace UE::TraceTools
+{
+
+void SFilterPresetList::Construct( const FArguments& InArgs )
+{
+	OnPresetChanged = InArgs._OnPresetChanged;
+	OnSavePreset = InArgs._OnSavePreset;
+	OnHighlightPreset = InArgs._OnHighlightPreset;
+
+	FilterBox = SNew(SWrapBox);
+
+	LoadSettings(FTraceToolsModule::TraceFiltersIni);
+	
+	PresetContainer = GetMutableDefault<ULocalTraceFilterPresetContainer>();
+	SharedPresetContainer = GetMutableDefault<USharedTraceFilterPresetContainer>();
+	EnginePresetContainer = GetMutableDefault<UEngineTraceFilterPresetContainer>();
+	
+	RefreshFilterPresets();
+
+	ChildSlot
+	[
+		FilterBox.ToSharedRef()
+	];
+
+}
+
+SFilterPresetList::~SFilterPresetList()
+{
+	SaveSettings(FTraceToolsModule::TraceFiltersIni);
+}
+
+void SFilterPresetList::RefreshFilterPresets()
+{
+	FSlateApplication::Get().DismissAllMenus();
+
+	if (HasAnyPresets())
+	{
+		for (auto FilterIt = Presets.CreateConstIterator(); FilterIt; ++FilterIt)
+		{
+			const TSharedRef<SFilterPreset>& FilterToRemove = *FilterIt;
+			ActiveFilterNames.Add(FilterToRemove->GetFilterPreset()->GetName());
+
+			if (FilterToRemove->IsEnabled())
+			{
+				EnabledFilterNames.Add(FilterToRemove->GetFilterPreset()->GetName());
+			}
+		}
+
+		FilterBox->ClearChildren();
+		Presets.Empty();
+	}
+
+	Presets.Empty();
+	EngineFilterPresets.Empty();
+	UserFilterPresets.Empty();
+	SharedUserFilterPresets.Empty();
+	AllFilterPresets.Empty();
+
+	PresetContainer->GetUserPresets(UserFilterPresets);
+
+	EnginePresetContainer->GetEnginePresets(EngineFilterPresets);
+	SharedPresetContainer->GetSharedUserPresets(SharedUserFilterPresets);
+
+	AllFilterPresets.Append(EngineFilterPresets);
+	AllFilterPresets.Append(UserFilterPresets);
+	AllFilterPresets.Append(SharedUserFilterPresets);
+
+	for (const TSharedPtr<ITraceFilterPreset>& Preset : AllFilterPresets)
+	{
+		if (ActiveFilterNames.Contains(Preset->GetName()))
+		{
+			TSharedRef<SFilterPreset> Filter = AddFilterPreset(Preset);
+			if (EnabledFilterNames.Contains(Preset->GetName()))
+			{
+				Filter->SetEnabled(true);
+			}
+		}
+	}
+
+	ActiveFilterNames.Empty();
+	EnabledFilterNames.Empty();
+}
+
+FReply SFilterPresetList::OnMouseButtonUp( const FGeometry& MyGeometry, const FPointerEvent& MouseEvent )
+{
+	if ( MouseEvent.GetEffectingButton() == EKeys::RightMouseButton )
+	{
+		FReply Reply = FReply::Handled().ReleaseMouseCapture();
+
+		// Get the context menu content. If NULL, don't open a menu.
+		TSharedPtr<SWidget> MenuContent = MakeFilterPresetsMenu();
+
+		if ( MenuContent.IsValid() )
+		{
+			FVector2D SummonLocation = MouseEvent.GetScreenSpacePosition();
+			FWidgetPath WidgetPath = MouseEvent.GetEventPath() != nullptr ? *MouseEvent.GetEventPath() : FWidgetPath();
+			FSlateApplication::Get().PushMenu(AsShared(), WidgetPath, MenuContent.ToSharedRef(), SummonLocation, FPopupTransitionEffect(FPopupTransitionEffect::ContextMenu));
+		}
+
+		return Reply;
+	}
+
+	return FReply::Unhandled();
+}
+
+void SFilterPresetList::SaveSettings(const FString& IniFilename) const
+{
+	FStringView IniSectionName = TEXTVIEW("FilterPresetList");
+	FStringView IniActivePresetsKey = TEXTVIEW("ActivePresets");
+	FStringView IniEnabledPresetsKey = TEXTVIEW("EnabledPresets");
+
+	FStringView Separator = TEXTVIEW(";");
+	
+	FString ActivePresetsString;
+	FString EnabledPresetsString;
+
+	for (const TSharedRef<SFilterPreset>& Preset : Presets)
+	{
+		ActivePresetsString.Append(Preset->GetFilterPreset()->GetName());
+		ActivePresetsString.Append(Separator);
+
+		if (Preset->IsEnabled())
+		{
+			EnabledPresetsString.Append(Preset->GetFilterPreset()->GetName());
+			EnabledPresetsString.Append(Separator);
+		}
+	}
+
+	ActivePresetsString.RemoveFromEnd(Separator);
+	EnabledPresetsString.RemoveFromEnd(Separator);
+
+	GConfig->SetString(IniSectionName.GetData(), IniActivePresetsKey.GetData(), *ActivePresetsString, IniFilename);
+	GConfig->SetString(IniSectionName.GetData(), IniEnabledPresetsKey.GetData(), *EnabledPresetsString, IniFilename);
+
+	GConfig->Flush(false, IniFilename);
+}
+
+void SFilterPresetList::LoadSettings(const FString& IniFilename)
+{
+	FStringView IniSectionName = TEXTVIEW("FilterPresetList");
+	FStringView IniActivePresetsKey = TEXT("ActivePresets");
+	FStringView IniEnabledPresetsKey = TEXT("EnabledPresets");
+
+	FString ActivePresetsString;
+	FString EnabledPresetsString;
+	GConfig->GetString(IniSectionName.GetData(), IniActivePresetsKey.GetData(), ActivePresetsString, IniFilename);
+	GConfig->GetString(IniSectionName.GetData(), IniEnabledPresetsKey.GetData(), EnabledPresetsString, IniFilename);
+
+	FStringView Separator = TEXT(";");
+	ActivePresetsString.ParseIntoArray(ActiveFilterNames, Separator.GetData());
+	EnabledPresetsString.ParseIntoArray(EnabledFilterNames, Separator.GetData());
+}
+
+bool SFilterPresetList::HasAnyPresets() const
+{
+	return Presets.Num() > 0;
+}
+
+void SFilterPresetList::EnableAllPresets()
+{
+	for (auto FilterIt = Presets.CreateConstIterator(); FilterIt; ++FilterIt)
+	{
+		(*FilterIt)->SetEnabled(true);
+	}
+}
+
+void SFilterPresetList::DisableAllPresets()
+{
+	for ( auto FilterIt = Presets.CreateConstIterator(); FilterIt; ++FilterIt )
+	{
+		(*FilterIt)->SetEnabled(false);
+	}
+}
+
+void SFilterPresetList::RemoveAllPresets()
+{
+	if ( HasAnyPresets() )
+	{
+		DisableAllPresets();
+		FilterBox->ClearChildren();
+		Presets.Empty();
+	}
+}
+
+
+void SFilterPresetList::GenerateEnginePresetsMenu(FMenuBuilder& InMenuBuilder)
+{
+	InMenuBuilder.BeginSection("TracePresetsEnginePresets", LOCTEXT("EnginePresetsMenuHeading", "Engine Presets"));
+	{
+		for (const TSharedPtr<ITraceFilterPreset>& EnginePreset : EngineFilterPresets)
+		{
+			InMenuBuilder.AddMenuEntry(
+				EnginePreset->GetDisplayText(),
+				EnginePreset->GetDisplayText(),
+				FSlateIcon(),
+				FUIAction(
+					FExecuteAction::CreateSP(this, &SFilterPresetList::TogglePreset, EnginePreset),
+					FCanExecuteAction(),
+					FIsActionChecked::CreateSP(this, &SFilterPresetList::IsPresetEnabled, EnginePreset)),
+				NAME_None,
+				EUserInterfaceActionType::ToggleButton
+			);
+		}
+	}
+	InMenuBuilder.EndSection();  // TracePresetsEnginePresets
+}
+
+void SFilterPresetList::GenerateLocalUserPresetsMenu(FMenuBuilder& InMenuBuilder)
+{
+	InMenuBuilder.BeginSection("TracePresetsUserPresets", LOCTEXT("UserPresetsMenuHeading", "User Presets"));
+	{
+		for (const TSharedPtr<ITraceFilterPreset>& UserPreset : UserFilterPresets)
+		{
+			InMenuBuilder.AddSubMenu(
+				UserPreset->GetDisplayText(),
+				UserPreset->GetDisplayText(),
+				FNewMenuDelegate::CreateLambda([this, UserPreset](FMenuBuilder& InSubMenuBuilder)
+				{
+					InSubMenuBuilder.BeginSection(NAME_None, LOCTEXT("UserPresetItemsHead", "User Preset(s)"));
+					{
+						GenerateCommonPresetEntries(InSubMenuBuilder, UserPreset);
+					}
+					InSubMenuBuilder.EndSection();
+				}),
+				FUIAction(
+					FExecuteAction::CreateSP(this, &SFilterPresetList::TogglePreset, UserPreset),
+					FCanExecuteAction(),
+					FIsActionChecked::CreateSP(this, &SFilterPresetList::IsPresetEnabled, UserPreset)),
+				NAME_None,
+				EUserInterfaceActionType::ToggleButton,
+				false, 
+				FSlateIcon(),
+				false
+			);
+		}
+	}
+	InMenuBuilder.EndSection(); // TracePresetsUserPresets
+}
+
+void SFilterPresetList::GenerateSharedUserPresetsMenu(FMenuBuilder& InMenuBuilder)
+{
+	InMenuBuilder.BeginSection("TracePresetsSharedUserPresets", LOCTEXT("SharedUserPresetsMenuHeading", "Shared User Presets"));
+	{
+		for (const TSharedPtr<ITraceFilterPreset>& UserPreset : SharedUserFilterPresets)
+		{
+			InMenuBuilder.AddSubMenu(
+				UserPreset->GetDisplayText(),
+				UserPreset->GetDisplayText(),
+				FNewMenuDelegate::CreateLambda([this, UserPreset](FMenuBuilder& InSubMenuBuilder)
+				{
+#if WITH_EDITOR
+					// SSettingsEditorCheckoutNotice is part of SharedWidgets module which depends on UnrealEd
+					TSharedPtr<SSettingsEditorCheckoutNotice> CheckoutWidget;
+					
+					SAssignNew(CheckoutWidget, SSettingsEditorCheckoutNotice)
+					.ConfigFilePath(GetDefault<USharedTraceFilterPresetContainer>()->GetDefaultConfigFilename())
+					.Visibility_Lambda([]() -> EVisibility
+					{
+						return FFilterPresetHelpers::CanModifySharedPreset() ? EVisibility::Collapsed : EVisibility::Visible;
+					});
+#endif // WITH_EDITOR
+					InSubMenuBuilder.BeginSection(NAME_None, LOCTEXT("SharedUserPresetItemHead", "Shared User Preset(s)"));
+					{
+#if WITH_EDITOR
+						InSubMenuBuilder.AddWidget(CheckoutWidget.ToSharedRef(), FText::GetEmpty());
+#endif // WITH_EDITOR
+						GenerateCommonPresetEntries(InSubMenuBuilder, UserPreset);
+					}
+					InSubMenuBuilder.EndSection();
+				}),
+				FUIAction(
+					FExecuteAction::CreateSP(this, &SFilterPresetList::TogglePreset, UserPreset),
+					FCanExecuteAction(),
+					FIsActionChecked::CreateSP(this, &SFilterPresetList::IsPresetEnabled, UserPreset)),
+				NAME_None,
+				EUserInterfaceActionType::ToggleButton,
+				false, 
+				FSlateIcon(),
+				false
+			);
+		}
+	}
+	InMenuBuilder.EndSection(); // TracePresetsSharedUserPresets
+}
+
+void SFilterPresetList::GenerateCommonPresetEntries(FMenuBuilder& InSubMenuBuilder, const TSharedPtr<ITraceFilterPreset>& UserPreset)
+{
+	InSubMenuBuilder.AddWidget(
+		SNew(SEditableTextBox)
+		.Text_Lambda([UserPreset]() { return UserPreset->GetDisplayText(); })
+		.OnTextCommitted_Lambda([UserPreset](const FText& InText, ETextCommit::Type InCommitType)
+		{
+			UserPreset->Rename(InText.ToString());
+		})
+		.OnVerifyTextChanged_Lambda([this, UserPreset](const FText& InNewText, FText& OutErrorMessage) -> bool
+		{
+			auto NameTestPredicate = [this, UserPreset, InNewText](TSharedPtr<ITraceFilterPreset> TestPreset)
+			{
+				return TestPreset != UserPreset && TestPreset->GetDisplayText().ToString() == InNewText.ToString();
+			};
+
+			if (UserFilterPresets.ContainsByPredicate(NameTestPredicate) || SharedUserFilterPresets.ContainsByPredicate(NameTestPredicate))
+			{
+				OutErrorMessage = LOCTEXT("DuplicatePresetNames", "This name is already in use");
+				return false;
+			}
+			return true;
+		})
+		, FText::GetEmpty()
+	);
+	
+	const FText DisplayText = UserPreset->IsLocal() ? LOCTEXT("MakeSharedPresetLabel", "Make Shared Preset") : LOCTEXT("MakeLocalPresetLabel", "Make Local Preset");
+	const FText TooltipText = UserPreset->IsLocal() ? LOCTEXT("MakeSharedPresetToolTip", "Makes this preset a Shared User Preset (Config INI file has to be writable)") : LOCTEXT("MakeLocalPresetToolTip", "Makes this preset a Local Preset");
+
+	InSubMenuBuilder.AddMenuEntry(DisplayText,
+		TooltipText,
+		FSlateIcon(),
+		FUIAction(
+			FExecuteAction::CreateLambda([UserPreset, this]()
+			{
+				const bool bSuccess = UserPreset->IsLocal() ? UserPreset->MakeShared() : UserPreset->MakeLocal();
+				if (bSuccess)
+				{
+					RefreshFilterPresets();
+				}
+			}),
+			FCanExecuteAction::CreateLambda([]()
+			{
+				return FFilterPresetHelpers::CanModifySharedPreset();
+			})
+		)
+	);
+
+		InSubMenuBuilder.AddMenuEntry(LOCTEXT("PresetSaveLabel", "Save Preset"),
+		LOCTEXT("PresetSaveToolTip", "Saves the current filtering state as this User Preset"),
+		FSlateIcon(),
+		FUIAction(
+			FExecuteAction::CreateLambda([UserPreset, this]()
+			{
+				if (UserPreset.IsValid())
+				{
+					OnSavePreset.ExecuteIfBound(UserPreset);
+				}
+			})
+		)
+	);
+
+	InSubMenuBuilder.AddMenuEntry(LOCTEXT("PresetDeleteLabel", "Delete Preset"),
+		LOCTEXT("PresetDeleteToolTip", "Deletes the User Preset"),
+		FSlateIcon(),
+		FUIAction(
+			FExecuteAction::CreateLambda([UserPreset, this]()
+			{
+				if ( UserPreset->Delete())
+				{
+					RefreshFilterPresets();
+				}
+			})
+		)
+	);
+}
+
+void SFilterPresetList::TogglePreset(const TSharedPtr<ITraceFilterPreset> InPreset)
+{
+	// Toggle off 
+	for (TSharedRef<SFilterPreset> Filter : Presets)
+	{
+		if (Filter->GetFilterPreset() == InPreset)
+		{
+			RemoveFilterPresetAndUpdate(Filter);
+			return;
+		}
+	}
+
+	// Toggle on
+	TSharedPtr<SFilterPreset> PresetWidget = AddFilterPreset(InPreset);
+	PresetWidget->SetEnabled(true);
+}
+
+bool SFilterPresetList::IsPresetEnabled(const TSharedPtr<ITraceFilterPreset> InPreset) const 
+{
+	for (TSharedRef<SFilterPreset> Filter : Presets)
+	{
+		if (Filter->GetFilterPreset() == InPreset)
+		{
+			return true;
+		}
+	}
+
+	return false;
+}
+
+TSharedRef<SWidget> SFilterPresetList::MakeFilterPresetsMenu()
+{
+	FMenuBuilder MenuBuilder(/*bInShouldCloseWindowAfterMenuSelection=*/false, nullptr, nullptr, /*bCloseSelfOnly=*/true);
+
+	MenuBuilder.BeginSection("TracePresetsResetPresets");
+	{
+		MenuBuilder.AddMenuEntry(
+			LOCTEXT("FilterListResetPresets", "Reset Presets"),
+			LOCTEXT("FilterListResetToolTip", "Resets current presets selection"),
+			FSlateIcon(),
+			FUIAction(FExecuteAction::CreateSP(this, &SFilterPresetList::RemoveAllPresets))
+		);
+		
+		MenuBuilder.AddMenuEntry(
+			LOCTEXT("FilterListSavePresets", "Save as User Preset"),
+			LOCTEXT("FilterListSaveToolTip", "Saves the currently filtering state as a new User Preset"),
+			FSlateIcon(),
+			FUIAction(FExecuteAction::CreateLambda([this]() 
+			{
+				// Cache all user names to determine which preset is the newly created one
+				TArray<FString> Names;
+				Algo::Transform(UserFilterPresets, Names, [](TSharedPtr<ITraceFilterPreset> Preset)
+				{
+					return Preset->GetName();
+				});
+
+				OnSavePreset.ExecuteIfBound(nullptr); 
+				RefreshFilterPresets();
+
+				for (const TSharedPtr<ITraceFilterPreset>& Preset : UserFilterPresets)
+				{
+					if (!Names.Contains(Preset->GetName()))
+					{
+						// This is the newly created preset, enable it for UX
+						TSharedRef<SFilterPreset> Widget = AddFilterPreset(Preset);
+						EnableOnlyThisPreset(Widget);
+						break;
+					}
+				}
+			}))
+		);
+	}
+	MenuBuilder.EndSection(); // TracePresetsResetPresets
+
+	GenerateEnginePresetsMenu(MenuBuilder);
+	GenerateLocalUserPresetsMenu(MenuBuilder);
+	GenerateSharedUserPresetsMenu(MenuBuilder);
+
+	return MenuBuilder.MakeWidget();
+}
+
+TSharedRef<SFilterPreset> SFilterPresetList::AddFilterPreset(const TSharedPtr<ITraceFilterPreset> FilterPreset)
+{
+	TSharedRef<SFilterPreset> NewFilter = SNew(SFilterPreset)
+		.FilterPreset(FilterPreset)
+		.OnPresetChanged(OnPresetChanged)
+		.OnRequestRemove(this, &SFilterPresetList::RemoveFilterPresetAndUpdate)
+		.OnRequestEnableAll(this, &SFilterPresetList::EnableAllPresets)
+		.OnRequestEnableOnly(this, &SFilterPresetList::EnableOnlyThisPreset)
+		.OnRequestDisableAll(this, &SFilterPresetList::DisableAllPresets)
+		.OnRequestRemoveAll(this, &SFilterPresetList::RemoveAllPresets)
+		.OnRequestDelete(this, & SFilterPresetList::DeletePreset)
+		.OnRequestSave(this, &SFilterPresetList::SavePreset)
+		.OnHighlightPreset(OnHighlightPreset);
+
+	Presets.Add(NewFilter);
+
+	// Add the widget as a child
+	FilterBox->AddSlot()
+	.Padding(3, 3)
+	[
+		NewFilter
+	];
+
+	return NewFilter;
+}
+
+void SFilterPresetList::EnableOnlyThisPreset(const TSharedRef<SFilterPreset>& PresetWidgetToEnable)
+{
+	for (TSharedRef<SFilterPreset>& PresetWidget : Presets )
+	{
+		const bool bEnable = (PresetWidget == PresetWidgetToEnable);
+		PresetWidget->SetEnabled(bEnable);
+	}
+}
+
+void SFilterPresetList::RemoveFilterPresetAndUpdate(const TSharedRef<SFilterPreset>& PresetToRemove)
+{
+	FilterBox->RemoveSlot(PresetToRemove);
+	Presets.RemoveSingleSwap(PresetToRemove);
+	PresetToRemove->SetEnabled(false);
+}
+
+void SFilterPresetList::DeletePreset(const TSharedRef<SFilterPreset>& PresetToDelete)
+{
+	if (PresetToDelete->GetFilterPreset()->Delete())
+	{
+		RefreshFilterPresets();
+	}
+}
+
+void SFilterPresetList::SavePreset(const TSharedRef<SFilterPreset>& PresetToSave)
+{
+	OnSavePreset.ExecuteIfBound(PresetToSave->GetFilterPreset());
+	EnableOnlyThisPreset(PresetToSave);
+}
+
+void SFilterPresetList::GetAllEnabledPresets(TArray<TSharedPtr<ITraceFilterPreset>>& OutPresets) const
+{
+	for (const TSharedRef<SFilterPreset>& Filter : Presets)
+	{
+		if (Filter->IsEnabled())
+		{
+			OutPresets.Add(Filter->GetFilterPreset());
+		}
+	}
+}
+
+TSharedRef<SWidget> SFilterPresetList::ExternalMakeFilterPresetsMenu()
+{
+	return MakeFilterPresetsMenu();
+}
+
+} // namespace UE::TraceTools
+
+#undef LOCTEXT_NAMESPACE
diff --git a/Engine/Source/Developer/TraceTools/Private/Widgets/SFilterPresetList.h b/Engine/Source/Developer/TraceTools/Private/Widgets/SFilterPresetList.h
new file mode 100644
--- /dev/null
+++ b/Engine/Source/Developer/TraceTools/Private/Widgets/SFilterPresetList.h
@@ -0,0 +1,133 @@
+// Copyright Epic Games, Inc. All Rights Reserved.
+
+#pragma once
+
+#include "SFilterPreset.h"
+
+class FMenuBuilder;
+class SWrapBox;
+
+class ULocalTraceFilterPresetContainer;
+class USharedTraceFilterPresetContainer;
+class UEngineTraceFilterPresetContainer;
+
+namespace UE::TraceTools
+{
+
+class SFilterPreset;
+
+/**
+ * A list of presets currently applied to insights filtering window
+ */
+class SFilterPresetList : public SCompoundWidget
+{
+public:
+	DECLARE_DELEGATE_OneParam(FOnSavePreset, const TSharedPtr<ITraceFilterPreset>& /*PresetToSave, nullptr means create a new one */);
+	DECLARE_DELEGATE_OneParam(FOnHighlightPreset, const TSharedPtr<ITraceFilterPreset>& /*PresetToHighlight, nullptr means stop highlighting */);
+
+	SLATE_BEGIN_ARGS( SFilterPresetList ){}
+		/** Delegate for when the selected presets have changed */
+		SLATE_EVENT(FOnPresetChanged, OnPresetChanged)
+		/** Delegate for when a preset is to be saved, given the current filtering state */
+		SLATE_EVENT(FOnSavePreset, OnSavePreset)
+		/** Delegate for when a specific preset is to be highlighted outside of this widget */
+		SLATE_EVENT(FOnHighlightPreset, OnHighlightPreset)
+	SLATE_END_ARGS()
+
+	/** Constructs this widget with InArgs */		
+	void Construct( const FArguments& InArgs );
+
+	virtual ~SFilterPresetList();
+
+	/** Begin SCompoundWidget overrides */
+	virtual FReply OnMouseButtonUp(const FGeometry& MyGeometry, const FPointerEvent& MouseEvent);
+	/** End SCompoundWidget overrides */
+
+	/** Returns all filter presets currently set as enabled within this widget */
+	void GetAllEnabledPresets(TArray<TSharedPtr<ITraceFilterPreset>>& OutPresets) const;
+
+	/** Returns the context menu widget generated by MakeFilterPresetsMenu */
+	TSharedRef<SWidget> ExternalMakeFilterPresetsMenu();	
+
+private:
+	/** Refreshes the filter preset objects and child widgets */
+	void RefreshFilterPresets();
+
+	/** Returns true if any Presets are applied */
+	bool HasAnyPresets() const;
+
+	/** Enables all Presets */
+	void EnableAllPresets();
+
+	/** Disables any applied Presets */
+	void DisableAllPresets();
+
+	/** Removes all Presets in the list */
+	void RemoveAllPresets();
+
+	/** Adds a child widget for the specified filter preset object */
+	TSharedRef<SFilterPreset> AddFilterPreset(const TSharedPtr<ITraceFilterPreset> FilterPreset);
+
+	/** Removes the child widget, and calls OnPresetsChanged to ensure the state is updated outside of this widget */
+	void RemoveFilterPresetAndUpdate(const TSharedRef<SFilterPreset>& PresetToRemove);
+
+	/** Handler for when the enable only this button was clicked on a single preset */
+	void EnableOnlyThisPreset(const TSharedRef<SFilterPreset>& PresetWidgetToEnable);
+
+	void TogglePreset(const TSharedPtr<ITraceFilterPreset> InPreset);
+	bool IsPresetEnabled(const TSharedPtr<ITraceFilterPreset> InPreset) const;
+
+	/** Saves any settings to config that should be persistent between editor sessions */
+	void SaveSettings(const FString& IniFilename) const;
+
+	/** Loads any settings to config that should be persistent between editor sessions */
+	void LoadSettings(const FString& IniFilename);
+
+	TSharedRef<SWidget> MakeFilterPresetsMenu();
+	void GenerateEnginePresetsMenu(FMenuBuilder& InMenuBuilder);
+	void GenerateLocalUserPresetsMenu(FMenuBuilder& InMenuBuilder);
+	void GenerateCommonPresetEntries(FMenuBuilder& InSubMenuBuilder, const TSharedPtr<ITraceFilterPreset>& UserPreset);
+	void GenerateSharedUserPresetsMenu(FMenuBuilder& InMenuBuilder);
+
+	/** Delete the preset object contained by the widget */
+	void DeletePreset(const TSharedRef<SFilterPreset>& PresetToDelete);
+	/** Save the preset object contained by the widget */
+	void SavePreset(const TSharedRef<SFilterPreset>& PresetToSave);	
+
+private:
+	/** The horizontal box which contains all the Presets */
+	TSharedPtr<SWrapBox> FilterBox;
+
+	/** All SFilterPreset in the list */
+	TArray<TSharedRef<SFilterPreset>> Presets;
+	
+	/** All Filter Preset objects */
+	TArray<TSharedPtr<ITraceFilterPreset>> AllFilterPresets;
+
+	/** All User Filter Preset objects */
+	TArray<TSharedPtr<ITraceFilterPreset>> UserFilterPresets;
+	/** All Shared User Filter Preset objects */
+	TArray<TSharedPtr<ITraceFilterPreset>> SharedUserFilterPresets;
+	/** All Engine Preset objects */
+	TArray<TSharedPtr<ITraceFilterPreset>> EngineFilterPresets;
+
+	/** Weak pointers to filter preset containers */
+	TWeakObjectPtr<ULocalTraceFilterPresetContainer> PresetContainer;
+	TWeakObjectPtr<USharedTraceFilterPresetContainer> SharedPresetContainer;
+	TWeakObjectPtr<UEngineTraceFilterPresetContainer> EnginePresetContainer;
+
+	/** Delegate for when Presets have changed */
+	FOnPresetChanged OnPresetChanged;
+
+	/** Delegate for when a Preset is to be saved */
+	FOnSavePreset OnSavePreset;
+
+	/** Delegate for when a Preset is to be highlighted */
+	FOnHighlightPreset OnHighlightPreset;
+
+	/** Cached arrays of currently active and enabled filter names, used for saving/restoring state during refreshes */
+	TArray<FString> ActiveFilterNames;
+	TArray<FString> EnabledFilterNames;
+};
+
+} // namespace UE::TraceTools
\ No newline at end of file
diff --git a/Engine/Source/Developer/TraceTools/Private/Widgets/STraceControl.cpp b/Engine/Source/Developer/TraceTools/Private/Widgets/STraceControl.cpp
new file mode 100644
--- /dev/null
+++ b/Engine/Source/Developer/TraceTools/Private/Widgets/STraceControl.cpp
@@ -0,0 +1,70 @@
+// Copyright Epic Games, Inc. All Rights Reserved.
+
+#include "STraceControl.h"
+
+#include "Framework/Commands/UICommandList.h"
+#include "STraceDataFilterWidget.h"
+#include "Widgets/SBoxPanel.h"
+#include "Widgets/STraceControlToolbar.h"
+
+namespace UE::TraceTools
+{
+
+STraceControl::STraceControl()
+{
+}
+
+STraceControl::~STraceControl()
+{
+
+}
+
+void STraceControl::Construct(const FArguments& InArgs, TSharedPtr<ITraceController> InTraceController)
+{
+	TraceController = InTraceController;
+
+	UICommandList = MakeShareable(new FUICommandList);
+	BindCommands();
+
+	TraceController->SendStatusUpdateRequest();
+	TraceController->SendChannelUpdateRequest();
+
+	ChildSlot
+	[
+		SNew(SVerticalBox)
+
+		+ SVerticalBox::Slot()
+		.AutoHeight()
+		[
+			SNew(STraceControlToolbar, UICommandList.ToSharedRef(), InTraceController)
+		]
+
+		+ SVerticalBox::Slot()
+		.Padding(FMargin(0.0f, 0.0f, 0.0f, 2.0f))
+		[
+			SNew(STraceDataFilterWidget, InTraceController)
+		]
+	];
+}
+
+void STraceControl::Tick(const FGeometry& AllottedGeometry, const double InCurrentTime, const float InDeltaTime)
+{
+	AccumulatedTime += InDeltaTime;
+
+	const double UpdateTime = 1.0f;
+
+	if (AccumulatedTime > UpdateTime)
+	{
+		TraceController->SendStatusUpdateRequest();
+		TraceController->SendChannelUpdateRequest();
+
+		AccumulatedTime = 0.0;
+	}
+}
+
+void STraceControl::BindCommands()
+{
+
+}
+
+} // namespace UE::TraceTools
diff --git a/Engine/Source/Developer/TraceTools/Private/Widgets/STraceControl.h b/Engine/Source/Developer/TraceTools/Private/Widgets/STraceControl.h
new file mode 100644
--- /dev/null
+++ b/Engine/Source/Developer/TraceTools/Private/Widgets/STraceControl.h
@@ -0,0 +1,42 @@
+// Copyright Epic Games, Inc. All Rights Reserved.
+
+#pragma once
+
+#include "ITraceController.h"
+#include "Widgets/SCompoundWidget.h"
+
+class ISessionManager;
+class FUICommandList;
+
+namespace UE::TraceTools
+{
+
+class STraceControl : public SCompoundWidget
+{
+public:
+	/** Default constructor. */
+	STraceControl();
+
+	/** Virtual destructor. */
+	virtual ~STraceControl();
+
+	SLATE_BEGIN_ARGS(STraceControl) {}
+	SLATE_END_ARGS()
+
+	/** Constructs this widget. */
+	void Construct(const FArguments& InArgs, TSharedPtr<ITraceController> InTraceController);
+
+	virtual void Tick(const FGeometry& AllottedGeometry, const double InCurrentTime, const float InDeltaTime) override;
+
+protected:
+	void BindCommands();
+
+private:
+	TSharedPtr<ITraceController> TraceController;
+
+	TSharedPtr<FUICommandList> UICommandList;
+
+	double AccumulatedTime = 0.0f;
+};
+
+} // namespace UE::TraceTools
diff --git a/Engine/Source/Developer/TraceTools/Private/Widgets/STraceControlToolbar.cpp b/Engine/Source/Developer/TraceTools/Private/Widgets/STraceControlToolbar.cpp
new file mode 100644
--- /dev/null
+++ b/Engine/Source/Developer/TraceTools/Private/Widgets/STraceControlToolbar.cpp
@@ -0,0 +1,312 @@
+// Copyright Epic Games, Inc. All Rights Reserved.
+
+#include "STraceControlToolbar.h"
+
+#include "ITraceController.h"
+#include "Framework/MultiBox/MultiBoxDefs.h"
+#include "Framework/MultiBox/MultiBoxBuilder.h"
+#include "SlateOptMacros.h"
+#include "SocketSubsystem.h"
+#include "Styling/AppStyle.h"
+#include "Widgets/Layout/SBorder.h"
+
+//TraceTools
+#include "Models/TraceControlCommands.h"
+
+#define LOCTEXT_NAMESPACE "STraceControlToolbar"
+
+namespace UE::TraceTools
+{
+
+STraceControlToolbar::STraceControlToolbar()
+{
+}
+
+STraceControlToolbar::~STraceControlToolbar()
+{
+	TraceController->OnStatusReceived().Remove(OnStatusReceivedDelegate);
+}
+
+BEGIN_SLATE_FUNCTION_BUILD_OPTIMIZATION
+void STraceControlToolbar::Construct(const FArguments& InArgs, const TSharedRef<FUICommandList>& CommandList, TSharedPtr<ITraceController> InTraceController)
+{
+	InitializeSettings();
+
+	FTraceControlCommands::Register();
+	TraceController = InTraceController;
+
+	OnStatusReceivedDelegate = TraceController->OnStatusReceived().AddSP(this, &STraceControlToolbar::OnTraceStatusUpdated);
+
+	BindCommands(CommandList);
+
+	// create the toolbar
+	FToolBarBuilder Toolbar(CommandList, FMultiBoxCustomization::None);
+	{
+		Toolbar.AddComboButton(
+			FUIAction(),
+			FOnGetContent::CreateSP(this, &STraceControlToolbar::BuildTraceTargetMenu, CommandList),
+			TAttribute<FText>::CreateSP(this, &STraceControlToolbar::GetTraceTargetLabelText),
+			TAttribute<FText>::CreateSP(this, &STraceControlToolbar::GetTraceTargetTooltipText),
+			TAttribute<FSlateIcon>::CreateSP(this, &STraceControlToolbar::GetTraceTargetIcon),
+			false);
+
+		Toolbar.AddSeparator();
+
+		Toolbar.AddToolBarButton(FTraceControlCommands::Get().StartTrace);
+		Toolbar.AddToolBarButton(FTraceControlCommands::Get().StopTrace);
+		Toolbar.AddToolBarButton(FTraceControlCommands::Get().TraceSnapshot);
+
+		Toolbar.AddSeparator();
+
+		Toolbar.AddToolBarButton(FTraceControlCommands::Get().PauseTrace);
+		Toolbar.AddToolBarButton(FTraceControlCommands::Get().ResumeTrace);
+
+		Toolbar.AddSeparator();
+
+		Toolbar.AddToolBarButton(FTraceControlCommands::Get().TraceBookmark);
+		Toolbar.AddToolBarButton(FTraceControlCommands::Get().TraceScreenshot);
+
+		Toolbar.AddSeparator();
+
+		Toolbar.AddToolBarButton(FTraceControlCommands::Get().ToggleStatNamedEvents);
+	}
+
+	ChildSlot
+	[
+		SNew(SBorder)
+		.BorderImage(FAppStyle::GetBrush("ToolPanel.GroupBorder"))
+		.Padding(0.0f)
+		[
+			Toolbar.MakeWidget()
+		]
+	];
+}
+END_SLATE_FUNCTION_BUILD_OPTIMIZATION
+
+TSharedRef<SWidget> STraceControlToolbar::BuildTraceTargetMenu(const TSharedRef<FUICommandList> CommandList)
+{
+	FMenuBuilder MenuBuilder(/*bInShouldCloseWindowAfterMenuSelection=*/true, CommandList);
+
+	MenuBuilder.SetSearchable(false);
+
+	MenuBuilder.AddMenuEntry(FTraceControlCommands::Get().SetTraceTargetServer);
+	MenuBuilder.AddMenuEntry(FTraceControlCommands::Get().SetTraceTargetFile);
+
+	return MenuBuilder.MakeWidget();
+}
+
+void STraceControlToolbar::BindCommands(const TSharedRef<FUICommandList>& CommandList)
+{
+	CommandList->MapAction(FTraceControlCommands::Get().SetTraceTargetServer, 
+						   FExecuteAction::CreateSP(this, &STraceControlToolbar::SetTraceTarget_Execute, ETraceTarget::Server),
+						   FCanExecuteAction::CreateSP(this, &STraceControlToolbar::SetTraceTarget_CanExecute));
+
+	CommandList->MapAction(FTraceControlCommands::Get().SetTraceTargetFile, 
+						   FExecuteAction::CreateSP(this, &STraceControlToolbar::SetTraceTarget_Execute, ETraceTarget::File),
+						   FCanExecuteAction::CreateSP(this, &STraceControlToolbar::SetTraceTarget_CanExecute));
+
+	CommandList->MapAction(FTraceControlCommands::Get().StartTrace, 
+					       FExecuteAction::CreateSP(this, &STraceControlToolbar::StartTrace_Execute),
+					       FCanExecuteAction::CreateSP(this, &STraceControlToolbar::StartTrace_CanExecute));
+
+	CommandList->MapAction(FTraceControlCommands::Get().StopTrace, 
+						   FExecuteAction::CreateSP(this, &STraceControlToolbar::StopTrace_Execute),
+						   FCanExecuteAction::CreateSP(this, &STraceControlToolbar::StopTrace_CanExecute));
+
+	CommandList->MapAction(FTraceControlCommands::Get().TraceSnapshot, 
+						   FExecuteAction::CreateSP(this, &STraceControlToolbar::TraceSnapshot_Execute),
+						   FCanExecuteAction::CreateSP(this, &STraceControlToolbar::TraceSnapshot_CanExecute));
+
+	CommandList->MapAction(FTraceControlCommands::Get().PauseTrace, 
+						   FExecuteAction::CreateSP(this, &STraceControlToolbar::PauseTrace_Execute),
+						   FCanExecuteAction::CreateSP(this, &STraceControlToolbar::PauseTrace_CanExecute));
+
+	CommandList->MapAction(FTraceControlCommands::Get().ResumeTrace, 
+						   FExecuteAction::CreateSP(this, &STraceControlToolbar::ResumeTrace_Execute),
+						   FCanExecuteAction::CreateSP(this, &STraceControlToolbar::ResumeTrace_CanExecute));
+
+	CommandList->MapAction(FTraceControlCommands::Get().TraceBookmark, 
+						   FExecuteAction::CreateSP(this, &STraceControlToolbar::TraceBookmark_Execute),
+						   FCanExecuteAction::CreateSP(this, &STraceControlToolbar::TraceBookmark_CanExecute));
+
+	CommandList->MapAction(FTraceControlCommands::Get().TraceScreenshot, 
+						   FExecuteAction::CreateSP(this, &STraceControlToolbar::TraceScreenshot_Execute),
+						   FCanExecuteAction::CreateSP(this, &STraceControlToolbar::TraceScreenshot_CanExecute));
+
+	CommandList->MapAction(FTraceControlCommands::Get().ToggleStatNamedEvents, 
+						   FExecuteAction::CreateSP(this, &STraceControlToolbar::ToggleStatNamedEvents_Execute),
+						   FCanExecuteAction::CreateSP(this, &STraceControlToolbar::ToggleStatNamedEvents_CanExecute),
+						   FIsActionChecked::CreateSP(this, &STraceControlToolbar::ToggleStatNamedEvents_IsChecked));
+}
+
+void STraceControlToolbar::InitializeSettings()
+{
+	TSharedPtr<FInternetAddr> RecorderAddr;
+	if (ISocketSubsystem* Sockets = ISocketSubsystem::Get())
+	{
+		bool bCanBindAll = false;
+		RecorderAddr = Sockets->GetLocalHostAddr(*GLog, bCanBindAll);
+	}
+
+	if (RecorderAddr.IsValid())
+	{
+		TraceHostAddr = RecorderAddr->ToString(false);
+	}
+	else
+	{
+		TraceHostAddr = TEXT("127.0.0.1");
+	}
+}
+
+bool STraceControlToolbar::SetTraceTarget_CanExecute() const
+{
+	return TraceController->HasAvailableSelectedInstance() && !bIsTracing;
+}
+
+void STraceControlToolbar::SetTraceTarget_Execute(ETraceTarget InTraceTarget)
+{
+	TraceTarget = InTraceTarget;
+}
+
+bool STraceControlToolbar::StartTrace_CanExecute() const
+{
+	return TraceController->HasAvailableSelectedInstance() && bIsTracing == false;
+}
+
+void STraceControlToolbar::StartTrace_Execute()
+{
+	if (TraceTarget == ETraceTarget::Server)
+	{
+		TraceController->Send(TraceHostAddr, TEXT(""));
+	}
+	else if (TraceTarget == ETraceTarget::File)
+	{
+		TraceController->File(TEXT(""), TEXT(""));
+	}
+	bIsTracing = true;
+}
+
+bool STraceControlToolbar::StopTrace_CanExecute() const
+{
+	return TraceController->HasAvailableSelectedInstance() && bIsTracing == true;
+}
+
+void STraceControlToolbar::StopTrace_Execute()
+{
+	TraceController->Stop();
+	bIsTracing = false;
+}
+
+bool STraceControlToolbar::TraceSnapshot_CanExecute() const
+{
+	return TraceController->HasAvailableSelectedInstance();
+}
+
+void STraceControlToolbar::TraceSnapshot_Execute()
+{
+	if (TraceTarget == ETraceTarget::Server)
+	{
+		TraceController->SnapshotSend(TraceHostAddr);
+	}
+	else if (TraceTarget == ETraceTarget::File)
+	{
+		TraceController->SnapshotFile(TEXT(""));
+	}
+}
+
+bool STraceControlToolbar::PauseTrace_CanExecute() const
+{
+	return TraceController->HasAvailableSelectedInstance() && bIsTracing && bIsPaused == false;
+}
+
+void STraceControlToolbar::PauseTrace_Execute()
+{
+	TraceController->Pause();
+	bIsPaused = true;
+}
+
+bool STraceControlToolbar::ResumeTrace_CanExecute() const
+{
+	return TraceController->HasAvailableSelectedInstance() && bIsTracing && bIsPaused == true;
+}
+
+void STraceControlToolbar::ResumeTrace_Execute()
+{
+	TraceController->Resume();
+	bIsPaused = false;
+}
+
+bool STraceControlToolbar::TraceBookmark_CanExecute() const
+{
+	return TraceController->HasAvailableSelectedInstance() && bIsTracing && !bIsPaused;
+}
+
+void STraceControlToolbar::TraceBookmark_Execute()
+{
+	const FString BookmarkName = FDateTime::Now().ToString(TEXT("Bookmark_%Y%m%d_%H%M%S"));
+	TraceController->Bookmark(BookmarkName);
+}
+
+bool STraceControlToolbar::TraceScreenshot_CanExecute() const
+{
+	return TraceController->HasAvailableSelectedInstance() && bIsTracing && !bIsPaused;
+}
+
+void STraceControlToolbar::TraceScreenshot_Execute()
+{
+	TraceController->Screenshot(TEXT(""), false);
+}
+
+bool STraceControlToolbar::ToggleStatNamedEvents_CanExecute() const
+{
+	return TraceController->HasAvailableSelectedInstance() && bIsTracing;
+}
+
+bool STraceControlToolbar::ToggleStatNamedEvents_IsChecked() const
+{
+	return bAreStatNamedEventsEnabled;
+}
+
+void STraceControlToolbar::ToggleStatNamedEvents_Execute()
+{
+	bAreStatNamedEventsEnabled = !bAreStatNamedEventsEnabled;
+}
+
+FText STraceControlToolbar::GetTraceTargetLabelText() const
+{
+	if (TraceTarget == ETraceTarget::Server)
+	{
+		return LOCTEXT("TraceTargetServerLabel", "Server");
+	}
+
+	return LOCTEXT("TraceTargetFileLabel", "File");
+}
+
+FText STraceControlToolbar::GetTraceTargetTooltipText() const
+{
+	if (TraceTarget == ETraceTarget::Server)
+	{
+		return LOCTEXT("TraceTargetServerTooltip", "Set the Unreal Trace Server as the trace target.");
+	}
+
+	return LOCTEXT("TraceTargetFileTooltip", "Set File as the trace target.");
+}
+
+FSlateIcon STraceControlToolbar::GetTraceTargetIcon() const
+{
+	if (TraceTarget == ETraceTarget::Server)
+	{
+		return FSlateIcon(FTraceToolsStyle::GetStyleSetName(), "TraceControl.SetTraceTargetServer");
+	}
+
+	return FSlateIcon(FTraceToolsStyle::GetStyleSetName(), "TraceControl.SetTraceTargetFile");
+}
+
+void STraceControlToolbar::OnTraceStatusUpdated(const FTraceStatus& InStatus, FTraceStatus::EUpdateType InUpdateType)
+{
+	bIsTracing = InStatus.bIsTracing;
+}
+
+} // namespace UE::TraceTools
+
+#undef LOCTEXT_NAMESPACE
\ No newline at end of file
diff --git a/Engine/Source/Developer/TraceTools/Private/Widgets/STraceControlToolbar.h b/Engine/Source/Developer/TraceTools/Private/Widgets/STraceControlToolbar.h
new file mode 100644
--- /dev/null
+++ b/Engine/Source/Developer/TraceTools/Private/Widgets/STraceControlToolbar.h
@@ -0,0 +1,90 @@
+// Copyright Epic Games, Inc. All Rights Reserved.
+
+#pragma once
+
+#include "ITraceController.h"
+#include "Framework/Commands/UICommandList.h"
+#include "Widgets/DeclarativeSyntaxSupport.h"
+#include "Widgets/SCompoundWidget.h"
+
+class ITraceController;
+
+namespace UE::TraceTools
+{
+
+enum ETraceTarget : uint8
+{
+	Server = 0,
+	File = 1
+};
+
+/**
+ * Implements the trace control toolbar widget.
+ */
+class STraceControlToolbar
+	: public SCompoundWidget
+{
+public:
+
+	SLATE_BEGIN_ARGS(STraceControlToolbar) { }
+	SLATE_END_ARGS()
+
+public:
+
+	STraceControlToolbar();
+	~STraceControlToolbar();
+
+	void Construct( const FArguments& InArgs, const TSharedRef<FUICommandList>& CommandList, TSharedPtr<ITraceController> InTraceController);
+
+private:
+	void BindCommands(const TSharedRef<FUICommandList>& CommandList);
+
+	TSharedRef<SWidget> BuildTraceTargetMenu(const TSharedRef<FUICommandList> CommandList);
+	FText GetTraceTargetLabelText() const;
+	FText GetTraceTargetTooltipText() const;
+	FSlateIcon GetTraceTargetIcon() const;
+
+	void OnTraceStatusUpdated(const FTraceStatus& InStatus, FTraceStatus::EUpdateType InUpdateType);
+	void InitializeSettings();
+
+	bool SetTraceTarget_CanExecute() const;
+	void SetTraceTarget_Execute(ETraceTarget InTraceTarget);
+
+	bool StartTrace_CanExecute() const;
+	void StartTrace_Execute();
+
+	bool StopTrace_CanExecute() const;
+	void StopTrace_Execute();
+
+	bool TraceSnapshot_CanExecute() const;
+	void TraceSnapshot_Execute();
+
+	bool PauseTrace_CanExecute() const;
+	void PauseTrace_Execute();
+
+	bool ResumeTrace_CanExecute() const;
+	void ResumeTrace_Execute();
+
+	bool TraceBookmark_CanExecute() const;
+	void TraceBookmark_Execute();
+
+	bool TraceScreenshot_CanExecute() const;
+	void TraceScreenshot_Execute();
+
+	bool ToggleStatNamedEvents_CanExecute() const;
+	bool ToggleStatNamedEvents_IsChecked() const;
+	void ToggleStatNamedEvents_Execute();
+
+private:
+	TSharedPtr<ITraceController> TraceController;
+
+	ETraceTarget TraceTarget = ETraceTarget::Server;
+	bool bIsTracing = false;
+	bool bIsPaused = false;
+	bool bAreStatNamedEventsEnabled = false;
+	FString TraceHostAddr;
+
+	FDelegateHandle OnStatusReceivedDelegate;
+};
+
+} // namespace UE::TraceTools
\ No newline at end of file
diff --git a/Engine/Source/Developer/TraceTools/Private/Widgets/STraceDataFilterWidget.cpp b/Engine/Source/Developer/TraceTools/Private/Widgets/STraceDataFilterWidget.cpp
new file mode 100644
--- /dev/null
+++ b/Engine/Source/Developer/TraceTools/Private/Widgets/STraceDataFilterWidget.cpp
@@ -0,0 +1,733 @@
+// Copyright Epic Games, Inc. All Rights Reserved.
+
+#include "STraceDataFilterWidget.h"
+
+#include "Framework/MultiBox/MultiBoxBuilder.h"
+#include "Framework/Views/TableViewMetadata.h"
+#include "Framework/Views/TreeFilterHandler.h"
+#include "Misc/TextFilter.h"
+#include "Modules/ModuleManager.h"
+#include "Widgets/Layout/SScrollBorder.h"
+#include "Widgets/Input/SComboButton.h"
+#include "Widgets/Layout/SScrollBar.h"
+#include "Widgets/Input/SSearchBox.h"
+#include "Widgets/Views/STreeView.h"
+
+// TraceTools
+#include "Models/TraceChannel.h"
+#include "Models/TraceFilterPresets.h"
+#include "Services/SessionTraceControllerFilterService.h"
+#include "STraceObjectRowWidget.h"
+#include "TraceToolsStyle.h"
+#include "Widgets/SFilterPresetList.h"
+
+#define LOCTEXT_NAMESPACE "STraceDataFilterWidget"
+
+namespace UE::TraceTools
+{
+
+TSet<FString> STraceDataFilterWidget::LastExpandedObjectNames;
+
+STraceDataFilterWidget::STraceDataFilterWidget() : PreviousSessionHandle(INDEX_NONE), bNeedsTreeRefresh(false), bHighlightingPreset(false)
+{
+}
+
+STraceDataFilterWidget::~STraceDataFilterWidget()
+{
+	// Save expansion and move to static array
+	SaveItemsExpansion();
+	LastExpandedObjectNames = ExpandedObjectNames;
+	ExpandedObjectNames.Empty();
+}
+
+void STraceDataFilterWidget::Construct(const FArguments& InArgs, TSharedPtr<ITraceController> InTraceController)
+{
+	TraceController = InTraceController;
+	SessionFilterService = MakeShareable(new FSessionTraceControllerFilterService(InTraceController));
+
+	SAssignNew(ExternalScrollbar, SScrollBar)
+	.AlwaysShowScrollbar(true);
+
+	ConstructSearchBoxFilter();
+	ConstructFilterHandler();
+	ConstructTreeview();
+
+	ChildSlot
+	[		
+		SNew(SBorder)
+		.Padding(4)
+		.BorderImage(FTraceToolsStyle::GetBrush("EventFilter.GroupBorder"))
+		[
+			SNew( SVerticalBox )
+
+			+ SVerticalBox::Slot()
+			.AutoHeight()
+			.Padding(FMargin(0.0f, 0.0f, 0.0f, 2.0f))
+			[
+				// Filtering button and search box widgets
+				SAssignNew(OptionsWidget, SHorizontalBox)
+				+SHorizontalBox::Slot()
+				.AutoWidth()
+				.VAlign(VAlign_Center)
+				.Padding(0.0f, 0.0f, 2.0f, 0.0f)
+				[
+					SNew(SComboButton)
+					.Visibility(EVisibility::Visible)
+					.ComboButtonStyle(FTraceToolsStyle::Get(), "EventFilter.ComboButton")
+					.ForegroundColor(FLinearColor::White)
+					.ContentPadding(0.0f)
+					.OnGetMenuContent(this, &STraceDataFilterWidget::MakeAddFilterMenu)
+					.ButtonContent()
+					[
+						SNew(SHorizontalBox)
+						+SHorizontalBox::Slot()
+						.AutoWidth()
+						.VAlign(VAlign_Center)
+						[
+							SNew(STextBlock)
+							.TextStyle(FTraceToolsStyle::Get(), "EventFilter.TextStyle")
+							.Font(FTraceToolsStyle::Get().GetFontStyle("FontAwesome.9"))
+							.Text(FText::FromString(FString(TEXT("\xf0b0"))) /*fa-filter*/)
+						]
+						+SHorizontalBox::Slot()
+						.AutoWidth()
+						.Padding(2, 0, 0, 0)
+						.VAlign(VAlign_Center)
+						[
+							SNew( STextBlock )
+							.TextStyle(FTraceToolsStyle::Get(), "EventFilter.TextStyle")
+							.Text( LOCTEXT("PresetsMenuLabel", "Filter Presets") )
+						]
+					]
+				]
+				+SHorizontalBox::Slot()
+				.FillWidth(1.0f)
+				[
+					SAssignNew(SearchBoxWidget, SSearchBox)
+					.SelectAllTextWhenFocused( true )
+					.HintText( LOCTEXT( "SearchBoxHint", "Search Trace Events..."))
+					.OnTextChanged(this, &STraceDataFilterWidget::OnSearchboxTextChanged)
+				]
+			]
+
+			+ SVerticalBox::Slot()
+			.Padding(0.0f, 0.0f, 0.0f, 2.0f)
+			.AutoHeight()
+			[
+				// Presets bar widget
+				SNew(SHorizontalBox)
+				+SHorizontalBox::Slot()
+				.FillWidth(1.0f)
+				[
+					SAssignNew(FilterPresetsListWidget, SFilterPresetList)
+					.OnPresetChanged(this, &STraceDataFilterWidget::OnPresetChanged)
+					.OnSavePreset(this, &STraceDataFilterWidget::OnSavePreset)
+					.OnHighlightPreset(this, &STraceDataFilterWidget::OnHighlightPreset)
+				]
+			]
+
+			+ SVerticalBox::Slot()
+			.AutoHeight()
+			[
+				SNew(SBorder)
+				.BorderImage(FTraceToolsStyle::GetBrush("FilterPresets.SessionWarningBorder"))
+				[
+					SNew(SHorizontalBox)
+					.Visibility_Lambda([this]() -> EVisibility 
+					{
+						return HasValidFilterSession() ? EVisibility::Collapsed : EVisibility::Visible;
+					})
+
+					+ SHorizontalBox::Slot()
+					.VAlign(VAlign_Center)
+					.AutoWidth()
+					.Padding(4.0f, 0.0f, 4.0f, 0.0f)
+					[
+						SNew(SImage)
+						.Image(FTraceToolsStyle::GetBrush("FilterPresets.WarningIcon"))
+					]
+
+					+ SHorizontalBox::Slot()
+					.VAlign(VAlign_Center)
+					.AutoWidth()
+					.Padding(4.0f, 0.0f, 0.0f, 0.0f)
+					[
+						SNew(STextBlock)
+						.Text(LOCTEXT("NoSessionSelectedWarning", "Please select an instance from the Session Browser."))
+						.ColorAndOpacity(FLinearColor(1.0f, 1.0f, 1.0f, 1.0f))
+					]
+				]
+			]
+
+			+ SVerticalBox::Slot()	
+			.Padding(0.0f, 0.0f, 0.0f, 2.0f)
+			.FillHeight(1.0f)
+			[
+				SNew(SOverlay)
+
+				+ SOverlay::Slot()
+				.HAlign(HAlign_Fill)
+				.VAlign(VAlign_Fill)
+				[
+					SNew(SHorizontalBox)
+
+					+ SHorizontalBox::Slot()
+					[
+						Treeview.ToSharedRef()
+					]
+
+					+ SHorizontalBox::Slot()
+					.AutoWidth()
+					[
+						SNew(SBox)
+						.WidthOverride(FOptionalSize(13.0f))
+						[
+							ExternalScrollbar.ToSharedRef()
+						]
+					]
+				]		
+			]
+		]
+	];
+
+	/** Setup attribute to enable/disabled all widgets according to whether or not there is a valid session to represent */
+	TAttribute<bool> EnabledAttribute;
+	EnabledAttribute.Bind(this, &STraceDataFilterWidget::HasValidFilterSession);
+
+	Treeview->SetEnabled(EnabledAttribute);
+	OptionsWidget->SetEnabled(EnabledAttribute);
+	FilterPresetsListWidget->SetEnabled(EnabledAttribute);
+
+	/** Restore expansion state when recreating this window */
+	ExpandedObjectNames = LastExpandedObjectNames;
+	RestoreItemsExpansion();
+}
+
+void STraceDataFilterWidget::OnSavePreset(const TSharedPtr<ITraceFilterPreset>& Preset)
+{
+	if (SessionFilterService.IsValid())
+	{
+		/** Save to preset if one was provided, otherwise create a new one */
+		if (Preset.IsValid())
+		{
+			Preset->Save(RootItems);
+		}
+		else
+		{
+			FFilterPresetHelpers::CreateNewPreset(RootItems);
+		}
+	}
+}
+
+void STraceDataFilterWidget::OnPresetChanged(const SFilterPreset& Preset)
+{
+	if (SessionFilterService.IsValid())
+	{
+		SessionFilterService->UpdateFilterPreset(Preset.GetFilterPreset(), Preset.IsEnabled());
+	}
+}
+
+void STraceDataFilterWidget::OnHighlightPreset(const TSharedPtr<ITraceFilterPreset>& Preset)
+{
+	if (Treeview.IsValid())
+	{
+		/** Update treeview so that any allowlisted entry (as part of Preset) is highlighted and expanded */
+		Treeview->ClearHighlightedItems();
+		if (Preset.IsValid())
+		{
+			if (!bHighlightingPreset)
+			{
+				/** Store current expansion, so we can reset once highlighting has finished */
+				SaveItemsExpansion();
+				bHighlightingPreset = true;
+			}
+
+			TArray<FString> Names;
+			Preset->GetAllowlistedNames(Names);
+
+			EnumerateAllItems([this, Names](TSharedPtr<ITraceObject> Object) -> void
+			{
+				if (Names.Contains(Object->GetName()))
+				{
+					Treeview->SetItemHighlighted(Object, true);
+					// TODO need to set parent expansion as well
+					Treeview->SetItemExpansion(Object, true);
+
+					SetParentExpansionRecursively(Object, true);
+				}
+			});
+		}
+		else
+		{
+			bHighlightingPreset = false;
+			RestoreItemsExpansion();
+		}
+	}
+}
+
+void STraceDataFilterWidget::OnSearchboxTextChanged(const FText& FilterText)
+{
+	bNeedsTreeRefresh = true;
+	const bool bFilterSet = !FilterText.IsEmpty();
+	if (bFilterSet != TreeviewFilterHandler->GetIsEnabled())
+	{
+		TreeviewFilterHandler->SetIsEnabled(bFilterSet);
+
+		if (bFilterSet)
+		{
+			SaveItemsExpansion();
+		}
+		else
+		{
+			RestoreItemsExpansion();
+		}
+	}
+
+	SearchBoxWidgetFilter->SetRawFilterText(FilterText);
+	SearchBoxWidget->SetError(SearchBoxWidgetFilter->GetFilterErrorText());
+}
+
+void STraceDataFilterWidget::ConstructTreeview()
+{
+	SAssignNew(Treeview, STreeView<TSharedPtr<ITraceObject>>)
+	.OnGetChildren(TreeviewFilterHandler.ToSharedRef(), &TreeFilterHandler<TSharedPtr<ITraceObject>>::OnGetFilteredChildren)
+	.OnGenerateRow(this, &STraceDataFilterWidget::OnGenerateRow)
+	.OnContextMenuOpening(this, &STraceDataFilterWidget::OnContextMenuOpening)
+	.OnMouseButtonDoubleClick(this, &STraceDataFilterWidget::OnItemDoubleClicked)
+	.ExternalScrollbar(ExternalScrollbar)
+	.OnSetExpansionRecursive_Lambda([this](TSharedPtr<ITraceObject> InObject, bool bInExpansionState)
+	{
+		SetExpansionRecursively(InObject, bInExpansionState);
+	})
+	.TreeItemsSource(&TreeItems);
+	
+	TreeviewFilterHandler->SetTreeView(Treeview.Get());
+}
+
+void STraceDataFilterWidget::OnItemDoubleClicked(TSharedPtr<ITraceObject> InObject) const
+{
+	TSharedPtr<ITableRow> Row = Treeview->WidgetFromItem(InObject);
+	if (Row.IsValid())
+	{
+		SetExpansionRecursively(InObject, !Row->IsItemExpanded());
+	}
+}
+
+void STraceDataFilterWidget::ConstructSearchBoxFilter()
+{
+	SearchBoxWidgetFilter = MakeShareable(new TTextFilter<TSharedPtr<ITraceObject>>(TTextFilter<TSharedPtr<ITraceObject>>::FItemToStringArray::CreateLambda([](TSharedPtr<ITraceObject> Object, TArray<FString>& OutStrings)
+	{
+		Object->GetSearchString(OutStrings);
+	})));
+}
+
+void STraceDataFilterWidget::ConstructFilterHandler()
+{
+	TreeviewFilterHandler = MakeShareable(new TreeFilterHandler<TSharedPtr<ITraceObject>>());
+	TreeviewFilterHandler->SetFilter(SearchBoxWidgetFilter.Get());
+	TreeviewFilterHandler->SetRootItems(&RootItems, &TreeItems);
+	TreeviewFilterHandler->SetGetChildrenDelegate(TreeFilterHandler<TSharedPtr<ITraceObject>>::FOnGetChildren::CreateLambda([](TSharedPtr<ITraceObject> InParent, TArray<TSharedPtr<ITraceObject>>& OutChildren)
+	{
+		InParent->GetChildren(OutChildren);
+	}));
+}
+
+TSharedRef<ITableRow> STraceDataFilterWidget::OnGenerateRow(TSharedPtr<ITraceObject> InItem, const TSharedRef<STableViewBase>& OwnerTable)
+{
+	LLM_SCOPE_BYNAME(TEXT("Insights/STraceDataFilterWidget"));
+	return SNew(STraceObjectRowWidget, OwnerTable, InItem)
+	.HighlightText(SearchBoxWidgetFilter.Get(), &TTextFilter<TSharedPtr<ITraceObject>>::GetRawFilterText);
+}
+
+TSharedRef<SWidget> STraceDataFilterWidget::MakeAddFilterMenu()
+{
+	return FilterPresetsListWidget->ExternalMakeFilterPresetsMenu();
+}
+
+TSharedPtr<SWidget> STraceDataFilterWidget::OnContextMenuOpening() const
+{
+	FMenuBuilder MenuBuilder(/*bInShouldCloseWindowAfterMenuSelection=*/true, nullptr, nullptr, /*bCloseSelfOnly=*/true);
+
+	static const FName FilteringSectionHook("FilteringState");
+	MenuBuilder.BeginSection(FilteringSectionHook, LOCTEXT("FilteringSectionLabel", "Filtering"));
+	{
+		MenuBuilder.AddMenuEntry(LOCTEXT("EnableAllRowsLabel", "Enable All"), LOCTEXT("EnableAllRowsTooltip", "Sets entire hierarchy to be non-filtered."), FSlateIcon(), 
+			FUIAction(FExecuteAction::CreateLambda([this]()
+				{
+					EnumerateAllItems([this](TSharedPtr<ITraceObject> Object) -> void
+					{
+						Object->SetIsFiltered(false);
+					});
+				}),
+				FCanExecuteAction(),
+				FGetActionCheckState(),
+				FIsActionButtonVisible::CreateLambda([this]() -> bool
+				{
+					return Treeview->GetNumItemsSelected() == 0 && EnumerateAllItems(TFunction<bool(TSharedPtr<ITraceObject> InItem)>([this](TSharedPtr<ITraceObject> Object) -> bool
+					{
+						return Object->IsFiltered();
+					}));
+				})
+			)
+		);
+
+		MenuBuilder.AddMenuEntry(LOCTEXT("DisableAllRowsLabel", "Disable All"), LOCTEXT("DisableAllRowsTooltip", "Sets entire hierarchy to be filtered."), FSlateIcon(), 
+			FUIAction(
+				FExecuteAction::CreateLambda([this]()
+				{
+					EnumerateAllItems([this](TSharedPtr<ITraceObject> Object) -> void
+					{
+						Object->SetIsFiltered(true);
+					});
+				}),
+				FCanExecuteAction(),
+				FGetActionCheckState(),
+				FIsActionButtonVisible::CreateLambda([this]() -> bool
+				{
+					return Treeview->GetNumItemsSelected() == 0 && EnumerateAllItems(TFunction<bool(TSharedPtr<ITraceObject> InItem)>([this](TSharedPtr<ITraceObject> Object) -> bool
+					{
+						return !Object->IsFiltered();
+					}));
+				})
+
+			)
+		);
+
+		MenuBuilder.AddMenuEntry(LOCTEXT("EnableRowsLabel", "Enable Selected"), LOCTEXT("EnableRowsTooltip", "Sets the selected Node(s) to be non-filtered."), FSlateIcon(),
+			FUIAction(
+				FExecuteAction::CreateRaw(this, &STraceDataFilterWidget::EnumerateSelectedItems, TFunction<void(TSharedPtr<ITraceObject> InItem)>([this](TSharedPtr<ITraceObject> Object) -> void
+				{
+					Object->SetIsFiltered(false);
+
+					SetExpansionRecursively(Object, true);
+				})),
+				FCanExecuteAction(),
+				FGetActionCheckState(),
+				FIsActionButtonVisible::CreateLambda([this]() -> bool
+				{
+					return EnumerateSelectedItems(TFunction<bool(TSharedPtr<ITraceObject> InItem)>([this](TSharedPtr<ITraceObject> Object) -> bool
+					{
+						return Object->IsFiltered();
+					}));
+				})
+			)
+		);
+
+		MenuBuilder.AddMenuEntry(LOCTEXT("DisableRowsLabel", "Disable Selected"), LOCTEXT("DisableRowsTooltip", "Sets the selected Node(s) to be filtered."), FSlateIcon(),
+			FUIAction(
+				FExecuteAction::CreateRaw(this, &STraceDataFilterWidget::EnumerateSelectedItems, TFunction<void(TSharedPtr<ITraceObject> InItem)>([this](TSharedPtr<ITraceObject> Object) -> void
+				{
+					Object->SetIsFiltered(true);
+				})),
+				FCanExecuteAction(),
+				FGetActionCheckState(),
+				FIsActionButtonVisible::CreateLambda([this]() -> bool
+				{
+					return EnumerateSelectedItems(TFunction<bool(TSharedPtr<ITraceObject> InItem)>([this](TSharedPtr<ITraceObject> Object) -> bool
+					{
+						return !Object->IsFiltered();
+					}));
+				})
+			)
+		);
+	}
+	MenuBuilder.EndSection();
+
+	static const FName ExpansionSectionHook("ExpansionState");
+	MenuBuilder.BeginSection(ExpansionSectionHook, LOCTEXT("ExpansionSectionLabel", "Expansion"));
+	{
+		MenuBuilder.AddMenuEntry(LOCTEXT("ExpandAllRowsLabel", "Expand All"), LOCTEXT("ExpandAllRowsTooltip", "Expands the entire hierarchy."), FSlateIcon(), FUIAction(
+				FExecuteAction::CreateLambda([this]()
+				{
+					for (const TSharedPtr<ITraceObject>& Object : TreeItems)
+					{
+						SetExpansionRecursively(Object, true);
+					}
+				}),
+				FCanExecuteAction(),
+				FGetActionCheckState(),
+				FIsActionButtonVisible::CreateLambda([this]() -> bool
+				{
+					return Treeview->GetNumItemsSelected() == 0 && EnumerateAllItems(TFunction<bool(TSharedPtr<ITraceObject> InItem)>([this](TSharedPtr<ITraceObject> Object) -> bool
+					{
+						return !Treeview->IsItemExpanded(Object);
+					}));
+				})
+			)
+		);
+		
+		MenuBuilder.AddMenuEntry(LOCTEXT("CollapseAllRowsLabel", "Collapse All"), LOCTEXT("CollapseAllRowsTooltip", "Collapses the entire hierarchy."), FSlateIcon(), 
+			FUIAction(
+				FExecuteAction::CreateLambda([this]()
+				{
+					for (const TSharedPtr<ITraceObject>& Object : TreeItems)
+					{
+						SetExpansionRecursively(Object, false);
+					}
+				}),
+				FCanExecuteAction(),
+				FGetActionCheckState(),
+				FIsActionButtonVisible::CreateLambda([this]() -> bool
+				{
+					return Treeview->GetNumItemsSelected() == 0 && EnumerateAllItems(TFunction<bool(TSharedPtr<ITraceObject> InItem)>([this](TSharedPtr<ITraceObject> Object) -> bool
+					{
+						return Treeview->IsItemExpanded(Object);
+					}));
+				})
+			)
+		);
+
+		MenuBuilder.AddMenuEntry(LOCTEXT("ExpandRowsLabel", "Expand Selected"), LOCTEXT("ExpandRowsTooltip", "Expands the selected Node(s)."), FSlateIcon(), 
+			FUIAction(
+				FExecuteAction::CreateRaw(this, &STraceDataFilterWidget::EnumerateSelectedItems, TFunction<void(TSharedPtr<ITraceObject> InItem)>([this](TSharedPtr<ITraceObject> Object) -> void
+				{
+					SetExpansionRecursively(Object, true);
+				})),
+				FCanExecuteAction(),
+				FGetActionCheckState(),
+				FIsActionButtonVisible::CreateLambda([this]() -> bool
+				{
+					return EnumerateSelectedItems(TFunction<bool(TSharedPtr<ITraceObject> InItem)>([this](TSharedPtr<ITraceObject> Object) -> bool
+					{
+						return !Treeview->IsItemExpanded(Object);
+					}));
+				})
+			)
+		);
+
+		MenuBuilder.AddMenuEntry(LOCTEXT("CollapseRowsLabel", "Collapse Selected"), LOCTEXT("CollapseRowsTooltip", "Collapse the selected Node(s)."), FSlateIcon(), 
+			FUIAction(
+				FExecuteAction::CreateRaw(this, &STraceDataFilterWidget::EnumerateSelectedItems, TFunction<void(TSharedPtr<ITraceObject> InItem)>([this](TSharedPtr<ITraceObject> Object) -> void
+				{
+					SetExpansionRecursively(Object, false);
+				})),
+				FCanExecuteAction(),
+				FGetActionCheckState(),
+				FIsActionButtonVisible::CreateLambda([this]() -> bool
+				{
+					return EnumerateSelectedItems(TFunction<bool(TSharedPtr<ITraceObject> InItem)>([this](TSharedPtr<ITraceObject> Object) -> bool
+					{
+						return Treeview->IsItemExpanded(Object);
+					}));
+				})
+			)
+		);
+	}  	
+	MenuBuilder.EndSection();
+
+	return MenuBuilder.MakeWidget();
+}
+
+void STraceDataFilterWidget::SaveItemsExpansion()
+{
+	ExpandedObjectNames.Empty();
+
+	if (Treeview.IsValid())
+	{
+		TSet<TSharedPtr<ITraceObject>> ExpandedItems;
+		Treeview->GetExpandedItems(ExpandedItems);
+
+		for (TSharedPtr<ITraceObject> Item : ExpandedItems)
+		{
+			if (Item.IsValid())
+			{
+				ExpandedObjectNames.Add(Item->GetName());
+			}
+		}
+	}
+}
+
+void STraceDataFilterWidget::RestoreItemsExpansion()
+{
+	EnumerateAllItems([this](TSharedPtr<ITraceObject> Object) -> void
+	{
+		const bool bExpanded = ExpandedObjectNames.Contains(Object->GetName());
+		SetExpansionRecursively(Object, bExpanded);
+	});
+
+	ExpandedObjectNames.Empty();
+}
+
+void STraceDataFilterWidget::SaveItemSelection()
+{
+	SelectedObjectNames.Empty();
+
+	if (Treeview.IsValid())
+	{
+		EnumerateSelectedItems([this](TSharedPtr<ITraceObject> InObject)
+		{
+			if (InObject.IsValid())
+			{
+				SelectedObjectNames.Add(InObject->GetName());
+			}
+		});
+	}
+}
+
+void STraceDataFilterWidget::RestoreItemSelection()
+{
+	if (Treeview.IsValid())
+	{
+		TArray<TSharedPtr<ITraceObject>> SelectedItems;
+		EnumerateAllItems([this, &SelectedItems](TSharedPtr<ITraceObject> Object) -> void
+		{
+			if (SelectedObjectNames.Contains(Object->GetName()))
+			{
+				SelectedItems.Add(Object);
+			}
+		});
+
+		Treeview->SetItemSelection(SelectedItems, true);
+	}
+	
+	SelectedObjectNames.Empty();
+}
+
+TSharedRef<ITraceObject> STraceDataFilterWidget::AddFilterableObject(const FTraceObjectInfo& Event, FString ParentName)
+{
+	// Retrieve any child objects, and recursively add those first 
+	TArray<FTraceObjectInfo> Events;
+	SessionFilterService->GetChildObjects(Event.Hash, Events);
+
+	TArray<TSharedPtr<ITraceObject>> ChildPtrs;	
+	for (const FTraceObjectInfo& ChildEvent : Events)
+	{
+		TSharedRef<ITraceObject> EventItem = AddFilterableObject(ChildEvent, Event.Name);
+		ChildPtrs.Add(EventItem);
+	}
+
+	TSharedRef<FTraceChannel> SharedItem = MakeShareable(new FTraceChannel(Event.Name, ParentName, Event.Hash, Event.bEnabled, Event.bReadOnly, ChildPtrs, SessionFilterService));
+	ParentToChild.Add(SharedItem, ChildPtrs);
+
+	for (TSharedPtr<ITraceObject> ChildObject : ChildPtrs)
+	{
+		ChildToParent.Add(ChildObject, SharedItem);
+	}
+
+	FlatItems.Add(SharedItem);
+
+	return SharedItem;
+}
+
+bool STraceDataFilterWidget::HasValidFilterSession() const
+{
+	return TraceController->HasAvailableSelectedInstance() && SessionFilterService.IsValid();
+}
+
+void STraceDataFilterWidget::SetExpansionRecursively(const TSharedPtr<ITraceObject>& InObject, bool bShouldExpandItem) const
+{
+	Treeview->SetItemExpansion(InObject, bShouldExpandItem);
+
+	TArray<TSharedPtr<ITraceObject>> Children;
+	InObject->GetChildren(Children);
+
+	for (TSharedPtr<ITraceObject>& ChildObject : Children)
+	{
+		SetExpansionRecursively(ChildObject, bShouldExpandItem);
+	}
+}
+
+void STraceDataFilterWidget::SetParentExpansionRecursively(const TSharedPtr<ITraceObject>& InObject, bool bShouldExpandItem) const
+{
+	if (const TSharedPtr<ITraceObject>* ParentObject = ChildToParent.Find(InObject))
+	{
+		Treeview->SetItemExpansion(*ParentObject, bShouldExpandItem);
+		SetParentExpansionRecursively(*ParentObject, bShouldExpandItem);
+	}
+}
+
+void STraceDataFilterWidget::EnumerateSelectedItems(TFunction<void(TSharedPtr<ITraceObject> InItem)> InFunction) const
+{
+	TArray<TSharedPtr<ITraceObject>> SelectedObjects;
+	Treeview->GetSelectedItems(SelectedObjects);
+
+	for (const TSharedPtr<ITraceObject>& Object : SelectedObjects)
+	{
+		InFunction(Object);
+	}
+}
+
+bool STraceDataFilterWidget::EnumerateSelectedItems(TFunction<bool(TSharedPtr<ITraceObject> InItem)> InFunction) const
+{
+	TArray<TSharedPtr<ITraceObject>> SelectedObjects;
+	Treeview->GetSelectedItems(SelectedObjects);
+
+
+	bool bState = false;
+	for (const TSharedPtr<ITraceObject>& Object : SelectedObjects)
+	{
+		bState |= InFunction(Object);
+	}
+
+	return bState;
+}
+
+void STraceDataFilterWidget::EnumerateAllItems(TFunction<void(TSharedPtr<ITraceObject> InItem)> InFunction) const
+{
+	for (const TSharedPtr<ITraceObject>& Object : FlatItems)
+	{
+		InFunction(Object);
+	}
+}
+
+bool STraceDataFilterWidget::EnumerateAllItems(TFunction<bool(TSharedPtr<ITraceObject> InItem)> InFunction) const
+{
+	bool bState = false;
+	for (const TSharedPtr<ITraceObject>& Object : FlatItems)
+	{
+		bState |= InFunction(Object);
+	}
+	return bState;
+}
+
+void STraceDataFilterWidget::RefreshTreeviewData()
+{
+	if (SessionFilterService.IsValid())
+	{
+		SyncTimeStamp = SessionFilterService->GetTimestamp();
+
+		/** Save expansion and selection */
+		SaveItemsExpansion();
+		SaveItemSelection();
+
+		TArray<FTraceObjectInfo> RootEvents;
+		SessionFilterService->GetRootObjects(RootEvents);
+		
+		ParentToChild.Empty();
+		ChildToParent.Empty();
+		RootItems.Empty();		
+
+		for (const FTraceObjectInfo& RootEvent : RootEvents)
+		{
+			TSharedRef<ITraceObject> TraceObject = AddFilterableObject(RootEvent, FString());
+			RootItems.Add(TraceObject);
+		}
+
+		TreeviewFilterHandler->RefreshAndFilterTree();
+
+		RestoreItemsExpansion();
+		RestoreItemSelection();
+	}
+}
+
+void STraceDataFilterWidget::Tick(const FGeometry& AllottedGeometry, const double InCurrentTime, const float InDeltaTime)
+{
+	LLM_SCOPE_BYNAME(TEXT("Insights/STraceDataFilterWidget"));
+
+	if (SessionFilterService.IsValid() )
+	{
+		if (SessionFilterService->GetTimestamp() != SyncTimeStamp)
+		{
+			RefreshTreeviewData();
+		}
+	}
+
+	if (bNeedsTreeRefresh)
+	{
+		TreeviewFilterHandler->RefreshAndFilterTree();
+	}
+}
+
+} // namespace UE::TraceTools
+
+#undef LOCTEXT_NAMESPACE // "STraceDataFilterWidget"
diff --git a/Engine/Source/Developer/TraceTools/Private/Widgets/STraceDataFilterWidget.h b/Engine/Source/Developer/TraceTools/Private/Widgets/STraceDataFilterWidget.h
new file mode 100644
--- /dev/null
+++ b/Engine/Source/Developer/TraceTools/Private/Widgets/STraceDataFilterWidget.h
@@ -0,0 +1,145 @@
+// Copyright Epic Games, Inc. All Rights Reserved.
+
+#pragma once
+
+#include "Widgets/SCompoundWidget.h"
+
+class SHorizontalBox;
+class STableViewBase;
+template <typename ItemType> class STreeView;
+template <typename ItemType> class TTextFilter;
+template <typename ItemType> class TreeFilterHandler;
+
+class ITableRow;
+class ITraceController;
+class SSearchBox;
+class SScrollBar;
+
+namespace UE::TraceTools
+{
+
+class SFilterPresetList;
+class SSessionSelector;
+struct ITraceFilterPreset;
+class ITraceObject;
+class ISessionTraceFilterService;
+class FEventFilterService;
+struct FTraceObjectInfo;
+
+class STraceDataFilterWidget : public SCompoundWidget
+{
+public:
+	/** Default constructor. */
+	STraceDataFilterWidget();
+
+	/** Virtual destructor. */
+	virtual ~STraceDataFilterWidget();
+
+	SLATE_BEGIN_ARGS(STraceDataFilterWidget) {}
+	SLATE_END_ARGS()
+
+	/** Constructs this widget. */
+	void Construct(const FArguments& InArgs, TSharedPtr<ITraceController> InTraceController);
+
+	virtual void Tick(const FGeometry& AllottedGeometry, const double InCurrentTime, const float InDeltaTime) override;
+
+protected:
+	/** Callback from SFilterPresetList, should save the current tree view filter state as the specified preset */
+	void OnSavePreset(const TSharedPtr<ITraceFilterPreset>& Preset);
+	/** Callback from SFilterPresetList, should update filtering state according to currently active Filter Presets */
+	void OnPresetChanged(const class SFilterPreset& Preset);
+	/** Callback from SFilterPresetList, should highlight the tree view items encompassed by the specified preset */
+	void OnHighlightPreset(const TSharedPtr<ITraceFilterPreset>& Preset);
+
+	void OnSearchboxTextChanged(const FText& FilterText);	
+	void OnItemDoubleClicked(TSharedPtr<ITraceObject> InObject) const;
+
+	/** Constructing various widgets */
+	void ConstructTreeview();
+	void ConstructSearchBoxFilter();
+	void ConstructFilterHandler();
+
+	bool HasValidFilterSession() const;
+		
+	void RefreshTreeviewData();
+		
+	/** Sets the treeview item expansion state recursively for the specified object */
+	void SetExpansionRecursively(const TSharedPtr<ITraceObject>& InObject, bool bShouldExpandItem) const;
+	void SetParentExpansionRecursively(const TSharedPtr<ITraceObject>& InObject, bool bShouldExpandItem) const;	
+
+	/** Enumerates the (currently selected) treeview items and calls InFunction for each entry */
+	void EnumerateSelectedItems(TFunction<void(TSharedPtr<ITraceObject> InItem)> InFunction) const;
+	void EnumerateAllItems(TFunction<void(TSharedPtr<ITraceObject> InItem)> InFunction) const;
+	bool EnumerateAllItems(TFunction<bool(TSharedPtr<ITraceObject> InItem)> InFunction) const;
+	bool EnumerateSelectedItems(TFunction<bool(TSharedPtr<ITraceObject> InItem)> InFunction) const;
+
+	/** Creates and adds a new treeview object entry, given the specified data */
+	TSharedRef<ITraceObject> AddFilterableObject(const FTraceObjectInfo& Event, FString ParentName);
+	
+	/** Treeview callbacks */
+	TSharedRef<ITableRow> OnGenerateRow(TSharedPtr<ITraceObject> InItem, const TSharedRef<STableViewBase>& OwnerTable);
+	TSharedPtr<SWidget> OnContextMenuOpening() const;
+
+	/** Presets dropdown button callback */
+	TSharedRef<SWidget> MakeAddFilterMenu();
+
+	/** (Re) storing item expansion / selection state */
+	void SaveItemsExpansion();
+	void SaveItemSelection();
+	void RestoreItemsExpansion();	
+	void RestoreItemSelection();
+
+protected:
+	/** Root level items used by treeview */
+	TArray<TSharedPtr<ITraceObject>> RootItems;
+	/** Dynamically generated array according to treeview filtering */
+	TArray<TSharedPtr<ITraceObject>> TreeItems;
+	/** Flat list of contained treeview items */
+	TArray<TSharedPtr<ITraceObject>> FlatItems;
+	/** Parent child relationships */
+	TMap<TSharedPtr<ITraceObject>, TArray<TSharedPtr<ITraceObject>>> ParentToChild;
+	TMap<TSharedPtr<ITraceObject>, TSharedPtr<ITraceObject>> ChildToParent;
+
+	/** Analysis filter service for the current analysis session this window is currently representing */
+	TSharedPtr<ISessionTraceFilterService> SessionFilterService;
+	uint32 PreviousSessionHandle;
+	/** Timestamp used for refreshing cached filter data */
+	FDateTime SyncTimeStamp = 0;
+	
+	/** Wrapper for presets drop down button */
+	TSharedPtr<SHorizontalBox> OptionsWidget;
+	
+	/** Treeview widget, represents the currently connected analysis session its filtering state */
+	TSharedPtr<STreeView<TSharedPtr<ITraceObject>>> Treeview;
+	TSharedPtr<TreeFilterHandler<TSharedPtr<ITraceObject>>> TreeviewFilterHandler;
+
+	/** Flag indicating the treeview should be refreshed */
+	bool bNeedsTreeRefresh;
+
+	/** Search box for filtering the Treeview items */
+	TSharedPtr<SSearchBox> SearchBoxWidget;
+	/** The text filter used by the search box */
+	TSharedPtr<TTextFilter<TSharedPtr<ITraceObject>>> SearchBoxWidgetFilter;
+
+	/** Filter presets bar widget */
+	TSharedPtr<SFilterPresetList> FilterPresetsListWidget;
+	
+	/** Widget used for selecting specific AnalysisSession the treeview should represent the filter state for */
+	TSharedPtr<SSessionSelector> SessionSelector;
+
+	/** Whether or not we a currently highlighting a preset */
+	bool bHighlightingPreset;
+
+	/** Cached state of which named entries were expanded / selected */
+	TSet<FString> ExpandedObjectNames;
+	TSet<FString> SelectedObjectNames;
+
+	/** Inter-window cached state of last expansion state */
+	static TSet<FString> LastExpandedObjectNames;
+
+	TSharedPtr<SScrollBar> ExternalScrollbar;
+
+	TSharedPtr<ITraceController> TraceController;
+};
+
+} // namespace UE::TraceTools
\ No newline at end of file
diff --git a/Engine/Source/Developer/TraceTools/Private/Widgets/STraceObjectRowWidget.cpp b/Engine/Source/Developer/TraceTools/Private/Widgets/STraceObjectRowWidget.cpp
new file mode 100644
--- /dev/null
+++ b/Engine/Source/Developer/TraceTools/Private/Widgets/STraceObjectRowWidget.cpp
@@ -0,0 +1,236 @@
+// Copyright Epic Games, Inc. All Rights Reserved.
+
+#include "STraceObjectRowWidget.h"
+
+// TraceTools
+#include "Models/ITraceObject.h"
+#include "TraceToolsStyle.h"
+
+namespace UE::TraceTools
+{
+
+class FSetFilteringDragDropOp : public FDragDropOperation
+{
+public:
+	DRAG_DROP_OPERATOR_TYPE(FSetFilteringDragDropOp, FDragDropOperation)
+
+	virtual ~FSetFilteringDragDropOp()
+	{
+		for (TSharedPtr<ITraceObject> Object : Objects)
+		{
+			Object->SetIsFiltered(bFilterOut);
+		}
+	}
+	
+	virtual TSharedPtr<SWidget> GetDefaultDecorator() const override
+	{
+		return SNullWidget::NullWidget;
+	}
+
+	static TSharedRef<FSetFilteringDragDropOp> New(TSharedPtr<ITraceObject> Object)
+	{
+		TSharedRef<FSetFilteringDragDropOp> Operation = MakeShareable(new FSetFilteringDragDropOp);
+		
+		Operation->AddObject(Object);
+		Operation->bFilterOut = !Object->IsFiltered();
+
+		Operation->Construct();
+		return Operation;
+	}
+
+	void AddObject(TSharedPtr<ITraceObject> InObject)
+	{
+		Objects.Add(InObject);
+		InObject->SetPending();
+	}
+
+	bool WillBeFilteredOut() const { return bFilterOut; }
+
+protected:
+	bool bFilterOut;
+	TSet<TSharedPtr<ITraceObject>> Objects;
+};
+
+
+class SFilterCheckboxWidget : public SImage
+{
+public:
+	SLATE_BEGIN_ARGS(SFilterCheckboxWidget) {}
+	SLATE_END_ARGS()
+
+	/** Construct this widget */
+	void Construct(const FArguments& InArgs, TSharedPtr<ITraceObject> InObject)
+	{
+		Object = InObject;
+		SImage::Construct(
+			SImage::FArguments()
+			.Image(this, &SFilterCheckboxWidget::GetBrush)
+			.ColorAndOpacity(FSlateColor::UseForeground())
+		);
+	}
+
+private:
+
+	/** Start a new drag/drop operation for this widget */
+	virtual FReply OnDragDetected(const FGeometry& MyGeometry, const FPointerEvent& MouseEvent) override
+	{
+		if (MouseEvent.IsMouseButtonDown(EKeys::LeftMouseButton))
+		{
+			return FReply::Handled().BeginDragDrop(FSetFilteringDragDropOp::New(Object));
+		}
+		else
+		{
+			return FReply::Unhandled();
+		}
+	}
+
+	/** If a visibility drag drop operation has entered this widget, set its actor to the new visibility state */
+	virtual void OnDragEnter(const FGeometry& MyGeometry, const FDragDropEvent& DragDropEvent) override
+	{
+		auto FilteringOp = DragDropEvent.GetOperationAs<FSetFilteringDragDropOp>();
+		if (FilteringOp.IsValid())
+		{
+			if (Object->IsFiltered() != FilteringOp->WillBeFilteredOut() && !Object->IsPending())
+			{
+				FilteringOp->AddObject(Object);
+			}
+		}
+	}
+
+	FReply HandleClick()
+	{
+		if (!IsPending())
+		{
+			return FReply::Handled().DetectDrag(SharedThis(this), EKeys::LeftMouseButton);
+		}
+
+		return FReply::Handled();		
+	}
+
+	virtual FReply OnMouseButtonDoubleClick(const FGeometry& InMyGeometry, const FPointerEvent& InMouseEvent) override
+	{
+		return HandleClick();
+	}
+
+	virtual FReply OnMouseButtonDown(const FGeometry& MyGeometry, const FPointerEvent& MouseEvent) override
+	{
+		if (MouseEvent.GetEffectingButton() != EKeys::LeftMouseButton)
+		{
+			return FReply::Unhandled();
+		}
+
+		return HandleClick();
+	}
+
+	virtual FReply OnMouseButtonUp(const FGeometry& MyGeometry, const FPointerEvent& MouseEvent) override
+	{
+		if (MouseEvent.GetEffectingButton() == EKeys::LeftMouseButton)
+		{
+			if (!IsPending())
+			{
+				ToggleIsFiltered();
+			}
+
+			return FReply::Handled();
+		}
+
+		return FReply::Unhandled();
+	}
+
+	const FSlateBrush* GetBrush() const
+	{
+		if (IsPending())
+		{
+			return IsHovered() ? FTraceToolsStyle::GetBrush("EventFilter.State.Pending_Hovered") :
+				FTraceToolsStyle::GetBrush("EventFilter.State.Pending");
+
+		}
+		else
+		{
+			if (IsFiltered())
+			{
+				return IsHovered() ? FTraceToolsStyle::GetBrush("EventFilter.State.Disabled_Hovered") :
+					FTraceToolsStyle::GetBrush("EventFilter.State.Disabled");
+			}
+			else
+			{
+				return IsHovered() ? FTraceToolsStyle::GetBrush("EventFilter.State.Enabled_Hovered") :
+					FTraceToolsStyle::GetBrush("EventFilter.State.Enabled");
+			}
+		}
+	}
+
+	bool IsFiltered() const
+	{
+		return Object->IsFiltered();
+	}
+
+	bool IsPending() const
+	{
+		return Object->IsPending();
+	}
+
+	void ToggleIsFiltered()
+	{
+		Object->SetIsFiltered(!Object->IsFiltered());
+	}
+	void SetIsFiltered(bool bState)
+	{
+		Object->SetIsFiltered(bState);
+	}
+
+	TSharedPtr<ITraceObject> Object;
+};
+
+void STraceObjectRowWidget::Construct(const FArguments& InArgs, const TSharedRef<STableViewBase>& InOwnerTableView, TSharedPtr<ITraceObject> InObject)
+{
+	Object = InObject;
+	HighlightTextAttribute = InArgs._HighlightText.Get();
+
+	STableRow<TSharedPtr<ITraceObject>>::Construct(
+		STableRow<TSharedPtr<ITraceObject>>::FArguments()
+		.Padding(0.0f)
+		.Content()
+		[
+			SNew(SHorizontalBox)
+			.IsEnabled(!InObject->IsReadOnly())
+			// Name 
+			+ SHorizontalBox::Slot()
+			.FillWidth(1.0f)
+			.Padding(2, 0, 0, 0)
+			.VAlign(VAlign_Center)
+			[
+				SNew(STextBlock)
+				.Text(Object->GetDisplayText())
+				.HighlightText_Lambda([this]() -> FText
+				{
+					// Get attribute value, populated by search box for treeview filter-ing
+					FText Value = HighlightTextAttribute.Get();
+
+					// In case it is empty allow for treeview highlighting to drive value
+					if (Value.IsEmpty())
+					{
+						auto OwnerTable = OwnerTablePtr.Pin().ToSharedRef();
+						bool bHighlight = OwnerTable->Private_IsItemHighlighted(Object);
+						if (bHighlight)
+						{
+							Value = Object->GetDisplayText();
+						}
+					}
+
+					return Value;
+				})
+			]
+
+			// Visibility icon
+			+ SHorizontalBox::Slot()
+			.AutoWidth()
+			.VAlign(VAlign_Center)
+			.Padding(0, 0, 2.f, 0)
+			[
+				SNew(SFilterCheckboxWidget, InObject)
+			]
+		], InOwnerTableView);
+}
+
+} // namespace UE::TraceTools
\ No newline at end of file
diff --git a/Engine/Source/Developer/TraceTools/Private/Widgets/STraceObjectRowWidget.h b/Engine/Source/Developer/TraceTools/Private/Widgets/STraceObjectRowWidget.h
new file mode 100644
--- /dev/null
+++ b/Engine/Source/Developer/TraceTools/Private/Widgets/STraceObjectRowWidget.h
@@ -0,0 +1,31 @@
+// Copyright Epic Games, Inc. All Rights Reserved.
+
+#pragma once
+
+#include "Widgets/Views/SListView.h"
+
+class STableViewBase;
+
+namespace UE::TraceTools
+{
+
+class ITraceObject;
+
+class STraceObjectRowWidget : public STableRow<TSharedPtr<ITraceObject>>
+{
+	SLATE_BEGIN_ARGS(STraceObjectRowWidget) {}
+	/** The current text to highlight */
+	SLATE_ATTRIBUTE(FText, HighlightText)
+	SLATE_END_ARGS()
+
+	void Construct(const FArguments& InArgs, const TSharedRef<STableViewBase>& InOwnerTableView, TSharedPtr<ITraceObject> InObject);
+	virtual ~STraceObjectRowWidget() {}
+
+protected:
+	/* The trace object that this tree item represents */
+	TSharedPtr<ITraceObject> Object;
+	TAttribute<FText> HighlightTextAttribute;
+
+};
+
+} // namespace UE::TraceTools
\ No newline at end of file
diff --git a/Engine/Source/Developer/TraceTools/Public/Interfaces/ITraceToolsModule.h b/Engine/Source/Developer/TraceTools/Public/Interfaces/ITraceToolsModule.h
new file mode 100644
--- /dev/null
+++ b/Engine/Source/Developer/TraceTools/Public/Interfaces/ITraceToolsModule.h
@@ -0,0 +1,41 @@
+// Copyright Epic Games, Inc. All Rights Reserved.
+
+/*=============================================================================
+	ITraceToolsModule.h: Declares the ITraceToolsModule interface.
+=============================================================================*/
+
+#pragma once
+
+#include "Modules/ModuleInterface.h"
+#include "Templates/SharedPointer.h"
+
+class ITraceController;
+class SWidget;
+
+namespace UE::TraceTools
+{
+
+/**
+ * Interface for a trace tools module.
+ */
+class ITraceToolsModule
+	: public IModuleInterface
+{
+public:
+	
+	/**
+	 * Creates a trace control widget.
+	 *
+	 * @return The new widget
+	 */
+	virtual TSharedRef<SWidget> CreateTraceControlWidget(TSharedPtr<ITraceController> InTraceController) = 0;
+
+public:
+
+	/**
+	 * Virtual destructor.
+	 */
+	virtual ~ITraceToolsModule( ) { }
+};
+
+} // namespace UE::TraceTools
diff --git a/Engine/Source/Developer/TraceTools/TraceTools.build.cs b/Engine/Source/Developer/TraceTools/TraceTools.build.cs
new file mode 100644
--- /dev/null
+++ b/Engine/Source/Developer/TraceTools/TraceTools.build.cs
@@ -0,0 +1,29 @@
+// Copyright Epic Games, Inc. All Rights Reserved.
+
+using UnrealBuildTool;
+
+public class TraceTools : ModuleRules
+{
+	public TraceTools(ReadOnlyTargetRules Target) : base(Target)
+	{
+		UnsafeTypeCastWarningLevel = WarningLevel.Error;
+
+		PrivateDependencyModuleNames.AddRange(
+				new string[] {
+					"Core",
+					"CoreUObject",
+					"InputCore",
+					"SessionServices",
+					"Slate",
+					"SlateCore",
+					"Sockets",
+					"TraceLog",
+				}
+			);
+
+		if (Target.bBuildEditor)
+		{
+			PrivateDependencyModuleNames.Add("SharedSettingsWidgets");
+		}
+	}
+}
diff --git a/Engine/Source/Programs/UnrealInsights/Private/UnrealInsightsMain.cpp b/Engine/Source/Programs/UnrealInsights/Private/UnrealInsightsMain.cpp
--- a/Engine/Source/Programs/UnrealInsights/Private/UnrealInsightsMain.cpp
+++ b/Engine/Source/Programs/UnrealInsights/Private/UnrealInsightsMain.cpp
@@ -18,8 +18,16 @@ int32 UnrealInsightsMain(const TCHAR* CommandLine)
 	// Override the stack size for the thread pool.
 	FQueuedThreadPool::OverrideStackSize = 256 * 1024;
 
+	FString NewCommandLine = CommandLine;
+
+	// Add -Messaging if it was not given in the command line.
+	if (!FParse::Param(*NewCommandLine, TEXT("Messaging")))
+	{
+		NewCommandLine += TEXT(" -Messaging");
+	}
+
 	// Initialize core.
-	GEngineLoop.PreInit(CommandLine);
+	GEngineLoop.PreInit(*NewCommandLine);
 
 	// Make sure all UObject classes are registered and default properties have been initialized.
 	//ProcessNewlyLoadedUObjects();
diff --git a/Engine/Source/Runtime/SessionServices/Private/TraceController.cpp b/Engine/Source/Runtime/SessionServices/Private/TraceController.cpp
--- a/Engine/Source/Runtime/SessionServices/Private/TraceController.cpp
+++ b/Engine/Source/Runtime/SessionServices/Private/TraceController.cpp
@@ -305,12 +305,30 @@ void FTraceController::OnSettings(const FTraceControlSettings& Message, const TS
 	}
 }
 
+bool FTraceController::HasAvailableSelectedInstance()
+{
+	if (SelectedInstanceId.IsValid())
+	{
+		return InstanceToAddress.Find(SelectedInstanceId) != nullptr;
+	}
+
+	return false;
+}
+
 void FTraceController::OnInstanceSelectionChanged(const TSharedPtr<ISessionInstanceInfo>& Instance, bool bSelected)
 {
 	FReadScopeLock _(InstancesLock);
 	
-	SelectedInstanceId = Instance->GetInstanceId();
-	if (const auto Address = InstanceToAddress.Find(SelectedInstanceId))
+	if (bSelected)
+	{
+		SelectedInstanceId = Instance->GetInstanceId();
+	}
+	else
+	{
+		SelectedInstanceId.Invalidate();
+	}
+
+	if (const auto Address = InstanceToAddress.Find(Instance->GetInstanceId()))
 	{
 		if (const auto Status = Instances.Find(*Address))
 		{
diff --git a/Engine/Source/Runtime/SessionServices/Private/TraceController.h b/Engine/Source/Runtime/SessionServices/Private/TraceController.h
--- a/Engine/Source/Runtime/SessionServices/Private/TraceController.h
+++ b/Engine/Source/Runtime/SessionServices/Private/TraceController.h
@@ -41,6 +41,8 @@ private:
 	virtual void SendStatusUpdateRequest() override;
 	virtual void SendChannelUpdateRequest() override;
 	virtual void SendSettingsUpdateRequest() override;
+
+	virtual bool HasAvailableSelectedInstance() override;
 	
 	DECLARE_DERIVED_EVENT(FTraceController, ITraceController::FStatusRecievedEvent, FStatusRecievedEvent);
 	virtual FStatusRecievedEvent& OnStatusReceived() override
diff --git a/Engine/Source/Runtime/SessionServices/Public/ITraceController.h b/Engine/Source/Runtime/SessionServices/Public/ITraceController.h
--- a/Engine/Source/Runtime/SessionServices/Public/ITraceController.h
+++ b/Engine/Source/Runtime/SessionServices/Public/ITraceController.h
@@ -187,5 +187,10 @@ public:
 	 * is updated. A reference to the status and what has changed is provided.
 	 */
 	virtual FStatusRecievedEvent& OnStatusReceived() = 0;
+
+	/**
+	 * Return true if a selected instance exists and has been discovered.
+	 */
+	virtual bool HasAvailableSelectedInstance() = 0;
 	
 };
