Fix Settings system for the Trace Control Window.

diff --git a/Engine/Config/BaseEngine.ini b/Engine/Config/BaseEngine.ini
--- a/Engine/Config/BaseEngine.ini
+++ b/Engine/Config/BaseEngine.ini
@@ -5,6 +5,9 @@
 [SectionsToSave]
 +Section=GameNetDriver StatelessConnectHandlerComponent
 +Section=/Script/InterchangeEngine.InterchangeProjectSettings
++Section=/Script/TraceTools.LocalTraceFilterPresetContainer
++Section=/Script/TraceTools.SharedTraceFilterPresetContainer
++Section=Trace.FilterPresetList
 
 [/Script/EngineSettings.GameMapsSettings]
 GameInstanceClass=/Script/Engine.GameInstance
@@ -2904,6 +2907,8 @@ DrawColor=(R=140,G=255,B=0,A=164)
 [Trace.ChannelPresets]
 Rendering=gpu,cpu,frame,log,bookmark
 Loading=cpu,frame,log,bookmark,screenshot,region,loadtime,assetloadtime,file
+TaskGraph=cpu,gpu,frame,log,bookmark,screenshot,region,task
+ContextSwitching=cpu,gpu,frame,log,bookmark,screenshot,region,contextswitch
 
 [RemoteConfiguration]
 Enabled=false
diff --git a/Engine/Source/Developer/TraceTools/Private/Models/TraceFilterPresets.cpp b/Engine/Source/Developer/TraceTools/Private/Models/TraceFilterPresets.cpp
--- a/Engine/Source/Developer/TraceTools/Private/Models/TraceFilterPresets.cpp
+++ b/Engine/Source/Developer/TraceTools/Private/Models/TraceFilterPresets.cpp
@@ -4,6 +4,7 @@
 
 #include "GenericPlatform/GenericPlatformFile.h"
 #include "HAL/PlatformFileManager.h"
+#include "Misc/StringBuilder.h"
 
 // TraceTools
 #include "ITraceObject.h"
@@ -38,14 +39,6 @@ void USharedTraceFilterPresetContainer::Save()
 	SharedPresetsContainer->TryUpdateDefaultConfigFile();
 }
 
-void UEngineTraceFilterPresetContainer::GetEnginePresets(TArray<TSharedPtr<UE::TraceTools::ITraceFilterPreset>>& OutPresets)
-{
-	for (FTraceFilterData& FilterData : EnginePresets)
-	{
-		OutPresets.Add(MakeShared<UE::TraceTools::FEngineFilterPreset>(FilterData.Name, FilterData));
-	}
-}
-
 void ULocalTraceFilterPresetContainer::GetUserPresets(TArray<TSharedPtr<UE::TraceTools::ITraceFilterPreset>>& OutPresets)
 {
 	for (FTraceFilterData& FilterData : UserPresets)
@@ -231,9 +224,24 @@ FText FFilterPreset::GetDisplayText() const
 	return FText::FromString(Name);
 }
 
-FText FFilterPreset::GetDescription() const
+FText FFilterPresetBase::GetDescription() const
 {
-	return FText::FormatOrdered(LOCTEXT("FilterPresetDescriptionFormat", "Name: {0}\nType: {1}"), FText::FromString(Name), CanDelete() ? (IsLocal() ? LOCTEXT("LocalPreset", "Local") : LOCTEXT("SharedPreset", "Shared")) : LOCTEXT("EnginePreset", "Engine"));
+	TStringBuilder<256> NamesList;
+	TArray<FString> Names;
+	GetAllowlistedNames(Names);
+
+	for (const FString& NameEntry : Names)
+	{
+		NamesList.Append(NameEntry);
+		NamesList.Append(TEXT(", "));
+	}
+	
+	if (NamesList.Len() > 1)
+	{
+		NamesList.RemoveSuffix(2);
+	}
+
+	return FText::FormatOrdered(LOCTEXT("FilterPresetDescriptionFormat", "Name: {0}\nType: {1}\nAllowlist: {2}"), FText::FromString(Name), CanDelete() ? (IsLocal() ? LOCTEXT("LocalPreset", "Local") : LOCTEXT("SharedPreset", "Shared")) : LOCTEXT("EnginePreset", "Engine"), FText::FromString(NamesList.ToString()));
 }
 
 void FFilterPreset::GetAllowlistedNames(TArray<FString>& OutNames) const
@@ -274,6 +282,11 @@ bool FFilterPreset::IsLocal() const
 	return false;
 }
 
+void FEngineFilterPreset::GetAllowlistedNames(TArray<FString>& OutNames) const
+{
+	OutNames.Insert(AllowListedNames, 0);
+}
+
 } // namespace UE::TraceTools
 
 #undef LOCTEXT_NAMESPACE // "FilterPreset"
diff --git a/Engine/Source/Developer/TraceTools/Private/Models/TraceFilterPresets.h b/Engine/Source/Developer/TraceTools/Private/Models/TraceFilterPresets.h
--- a/Engine/Source/Developer/TraceTools/Private/Models/TraceFilterPresets.h
+++ b/Engine/Source/Developer/TraceTools/Private/Models/TraceFilterPresets.h
@@ -38,7 +38,7 @@ struct FTraceFilterData
 };
 
 /** UObject containers for the preset data */
-UCLASS(Config = Trace)
+UCLASS(Config = Engine)
 class ULocalTraceFilterPresetContainer : public UObject
 {
 	GENERATED_BODY()
@@ -55,7 +55,7 @@ protected:
 	TArray<FTraceFilterData> UserPresets;
 };
 
-UCLASS(Config = Trace, DefaultConfig)
+UCLASS(Config = Engine, DefaultConfig)
 class USharedTraceFilterPresetContainer : public UObject
 {
 	GENERATED_BODY()
@@ -72,32 +72,28 @@ protected:
 	TArray<FTraceFilterData> SharedPresets;
 };
 
-UCLASS(Config = Trace, DefaultConfig)
-class UEngineTraceFilterPresetContainer : public UObject
+namespace UE::TraceTools
 {
-	GENERATED_BODY()
 
-	friend struct UE::TraceTools::FFilterPresetHelpers;
-public:
-	void GetEnginePresets(TArray<TSharedPtr<UE::TraceTools::ITraceFilterPreset>>& OutPresets);
+/** Base implementation of a filter preset */
+struct FFilterPresetBase : public ITraceFilterPreset
+{
 protected:
-	UPROPERTY(Config)
-	TArray<FTraceFilterData> EnginePresets;
-};
+	FFilterPresetBase(const FString& InName) : Name(InName) {}
+	virtual FText GetDescription() const;
 
-namespace UE::TraceTools
-{
+protected:
+	FString Name;
+};
 
-/** Base implementation of a filter preset */
-struct FFilterPreset : public TraceTools::ITraceFilterPreset
+struct FFilterPreset : public FFilterPresetBase
 {
 public:
-	FFilterPreset(const FString& InName, FTraceFilterData& InFilterData) : Name(InName), FilterData(InFilterData) {}
+	FFilterPreset(const FString& InName, FTraceFilterData& InFilterData) : FFilterPresetBase(InName), FilterData(InFilterData) {}
 
 	/** Begin ITraceFilterPreset overrides */
 	virtual FString GetName() const override;
 	virtual FText GetDisplayText() const;
-	virtual FText GetDescription() const;
 	virtual void GetAllowlistedNames(TArray<FString>& OutNames) const override;
 	virtual bool CanDelete() const override;
 	virtual void Rename(const FString& InNewName) override;
@@ -109,12 +105,32 @@ public:
 	virtual void Save() override {}
 	/** End ITraceFilterPreset overrides */
 protected:
-	FString Name;
 	FTraceFilterData& FilterData;
 };
 
-/** Engine level preset is simply a basic one */
-typedef FFilterPreset FEngineFilterPreset;
+// The Engine presets don't use USTRUCT so they can be shared with FTraceAuxiliary.
+struct FEngineFilterPreset : public FFilterPresetBase
+{
+public:
+	FEngineFilterPreset(const FString& InName, const TArray<FString>& InAllowListedNames) : FFilterPresetBase(InName), AllowListedNames(InAllowListedNames) {}
+
+	/** Begin ITraceFilterPreset overrides */
+	virtual FString GetName() const override { return Name; }
+	virtual FText GetDisplayText() const { return FText::FromString(Name); }
+	virtual void GetAllowlistedNames(TArray<FString>& OutNames) const override;
+	virtual bool CanDelete() const override { return false; }
+	virtual void Rename(const FString& InNewName) override {}
+	virtual bool Delete() override { return false; }
+	virtual bool MakeShared() override { return false; }
+	virtual bool MakeLocal() override { return false; }
+	virtual bool IsLocal() const override { return false; }
+	virtual void Save(const TArray<TSharedPtr<TraceTools::ITraceObject>>& InObjects) override {}
+	virtual void Save() override {}
+	/** End ITraceFilterPreset overrides */
+
+protected:
+	TArray<FString> AllowListedNames;
+};
 
 /** User filter preset, allows for deletion / transitioning INI ownership */
 struct FUserFilterPreset : public FFilterPreset
diff --git a/Engine/Source/Developer/TraceTools/Private/Services/SessionTraceControllerFilterService.cpp b/Engine/Source/Developer/TraceTools/Private/Services/SessionTraceControllerFilterService.cpp
--- a/Engine/Source/Developer/TraceTools/Private/Services/SessionTraceControllerFilterService.cpp
+++ b/Engine/Source/Developer/TraceTools/Private/Services/SessionTraceControllerFilterService.cpp
@@ -59,10 +59,20 @@ void FSessionTraceControllerFilterService::UpdateFilterPreset(const TSharedPtr<I
 	if (IsEnabled)
 	{
 		FrameEnabledChannels.Append(Names);
+
+		for (const FString& Name : Names)
+		{
+			FrameDisabledChannels.Remove(Name);
+		}
 	}
 	else
 	{
 		FrameDisabledChannels.Append(Names);
+
+		for (const FString& Name : Names)
+		{
+			FrameEnabledChannels.Remove(Name);
+		}
 	}
 }
 
@@ -94,6 +104,11 @@ void FSessionTraceControllerFilterService::OnTraceStatusUpdated(const FTraceStat
 	const TMap<uint32, FTraceStatus::FChannel> Channels = InStatus.Channels;
 	Objects.Empty(Channels.Num());
 
+	if (Channels.Num())
+	{
+		bChannelsReceived = true;
+	}
+
 	for (auto& Entry : Channels)
 	{
 		FTraceObjectInfo& EventInfo = Objects.AddDefaulted_GetRef();
@@ -107,15 +122,21 @@ void FSessionTraceControllerFilterService::OnTraceStatusUpdated(const FTraceStat
 
 void FSessionTraceControllerFilterService::OnApplyChannelChanges()
 {
+	if (!TraceController->HasAvailableSelectedInstance() || !bChannelsReceived)
+	{
+		return;
+	}
+
 	if (FrameEnabledChannels.Num() || FrameDisabledChannels.Num())
 	{
 		TraceController->WithSelectedInstances([&](ITraceControllerCommands& Commands)
 		{
-			Commands.SetChannels(FrameEnabledChannels, FrameDisabledChannels);
+			Commands.SetChannels(FrameEnabledChannels.Array(), FrameDisabledChannels.Array());
 		});
+
 		FrameEnabledChannels.Empty();
 		FrameDisabledChannels.Empty();
 	}
-}
+} 
 
 } // namespace UE::TraceTools
\ No newline at end of file
diff --git a/Engine/Source/Developer/TraceTools/Private/Services/SessionTraceControllerFilterService.h b/Engine/Source/Developer/TraceTools/Private/Services/SessionTraceControllerFilterService.h
--- a/Engine/Source/Developer/TraceTools/Private/Services/SessionTraceControllerFilterService.h
+++ b/Engine/Source/Developer/TraceTools/Private/Services/SessionTraceControllerFilterService.h
@@ -49,11 +49,13 @@ protected:
 	TArray<FTraceObjectInfo> Objects;
 
 	/** Names of channels that were either enabled or disabled during the duration of this frame */
-	TArray<FString> FrameEnabledChannels;
-	TArray<FString> FrameDisabledChannels;
+	TSet<FString> FrameEnabledChannels;
+	TSet<FString> FrameDisabledChannels;
 
 	/** Timestamp at which contained data (including provider) was last updated */
-	FDateTime TimeStamp;	
+	FDateTime TimeStamp;
+
+	bool bChannelsReceived = false;
 };
 
 } // namespace UE::TraceTools
\ No newline at end of file
diff --git a/Engine/Source/Developer/TraceTools/Private/Widgets/SFilterPresetList.cpp b/Engine/Source/Developer/TraceTools/Private/Widgets/SFilterPresetList.cpp
--- a/Engine/Source/Developer/TraceTools/Private/Widgets/SFilterPresetList.cpp
+++ b/Engine/Source/Developer/TraceTools/Private/Widgets/SFilterPresetList.cpp
@@ -28,14 +28,14 @@ void SFilterPresetList::Construct( const FArguments& InArgs )
 	OnSavePreset = InArgs._OnSavePreset;
 	OnHighlightPreset = InArgs._OnHighlightPreset;
 
-	FilterBox = SNew(SWrapBox);
+	FilterBox = SNew(SWrapBox)
+				.UseAllottedSize(true);
 
 	LoadSettings(FTraceToolsModule::TraceFiltersIni);
 	
 	PresetContainer = GetMutableDefault<ULocalTraceFilterPresetContainer>();
 	SharedPresetContainer = GetMutableDefault<USharedTraceFilterPresetContainer>();
-	EnginePresetContainer = GetMutableDefault<UEngineTraceFilterPresetContainer>();
-	
+
 	RefreshFilterPresets();
 
 	ChildSlot
@@ -77,9 +77,9 @@ void SFilterPresetList::RefreshFilterPresets()
 	SharedUserFilterPresets.Empty();
 	AllFilterPresets.Empty();
 
+	LoadEnginePresets();
+	
 	PresetContainer->GetUserPresets(UserFilterPresets);
-
-	EnginePresetContainer->GetEnginePresets(EngineFilterPresets);
 	SharedPresetContainer->GetSharedUserPresets(SharedUserFilterPresets);
 
 	AllFilterPresets.Append(EngineFilterPresets);
@@ -126,7 +126,7 @@ FReply SFilterPresetList::OnMouseButtonUp( const FGeometry& MyGeometry, const FP
 
 void SFilterPresetList::SaveSettings(const FString& IniFilename) const
 {
-	FStringView IniSectionName = TEXTVIEW("FilterPresetList");
+	FStringView IniSectionName = TEXTVIEW("Trace.FilterPresetList");
 	FStringView IniActivePresetsKey = TEXTVIEW("ActivePresets");
 	FStringView IniEnabledPresetsKey = TEXTVIEW("EnabledPresets");
 
@@ -158,7 +158,7 @@ void SFilterPresetList::SaveSettings(const FString& IniFilename) const
 
 void SFilterPresetList::LoadSettings(const FString& IniFilename)
 {
-	FStringView IniSectionName = TEXTVIEW("FilterPresetList");
+	FStringView IniSectionName = TEXTVIEW("Trace.FilterPresetList");
 	FStringView IniActivePresetsKey = TEXT("ActivePresets");
 	FStringView IniEnabledPresetsKey = TEXT("EnabledPresets");
 
@@ -534,6 +534,36 @@ TSharedRef<SWidget> SFilterPresetList::ExternalMakeFilterPresetsMenu()
 	return MakeFilterPresetsMenu();
 }
 
+void SFilterPresetList::LoadEnginePresets()
+{
+	TArray<FString> PresetStrings;
+	GConfig->GetSection(TEXT("Trace.ChannelPresets"), PresetStrings, GEngineIni);
+
+	for (const FString& Item : PresetStrings)
+	{
+		FString Key, Value;
+		Item.Split(TEXT("="), &Key, &Value);
+
+		TArray<FString> AllowListedNames;
+		Value.ParseIntoArray(AllowListedNames, TEXT(","));
+		EngineFilterPresets.Add(MakeShared<FEngineFilterPreset>(Key, AllowListedNames));
+	}
+
+	TArray<FTraceAuxiliary::FChannelPreset> EnginePresets;
+	FTraceAuxiliary::GetFixedChannelPresets(EnginePresets);
+
+	for (auto& Item : EnginePresets)
+	{
+		if (!Item.bIsReadOnly)
+		{
+			TArray<FString> AllowListedNames;
+			FString Channels = Item.Channels;
+			Channels.ParseIntoArray(AllowListedNames, TEXT(","));
+			EngineFilterPresets.Add(MakeShared<FEngineFilterPreset>(Item.Name, AllowListedNames));
+		}
+	}
+}
+
 } // namespace UE::TraceTools
 
 #undef LOCTEXT_NAMESPACE
diff --git a/Engine/Source/Developer/TraceTools/Private/Widgets/SFilterPresetList.h b/Engine/Source/Developer/TraceTools/Private/Widgets/SFilterPresetList.h
--- a/Engine/Source/Developer/TraceTools/Private/Widgets/SFilterPresetList.h
+++ b/Engine/Source/Developer/TraceTools/Private/Widgets/SFilterPresetList.h
@@ -9,7 +9,6 @@ class SWrapBox;
 
 class ULocalTraceFilterPresetContainer;
 class USharedTraceFilterPresetContainer;
-class UEngineTraceFilterPresetContainer;
 
 namespace UE::TraceTools
 {
@@ -94,6 +93,9 @@ private:
 	/** Save the preset object contained by the widget */
 	void SavePreset(const TSharedRef<SFilterPreset>& PresetToSave);	
 
+	/** Engine presets have a custom loading so they can be shared with the console commands from FTraceAuxiliary. */
+	void LoadEnginePresets();
+
 private:
 	/** The horizontal box which contains all the Presets */
 	TSharedPtr<SWrapBox> FilterBox;
@@ -114,7 +116,6 @@ private:
 	/** Weak pointers to filter preset containers */
 	TWeakObjectPtr<ULocalTraceFilterPresetContainer> PresetContainer;
 	TWeakObjectPtr<USharedTraceFilterPresetContainer> SharedPresetContainer;
-	TWeakObjectPtr<UEngineTraceFilterPresetContainer> EnginePresetContainer;
 
 	/** Delegate for when Presets have changed */
 	FOnPresetChanged OnPresetChanged;
diff --git a/Engine/Source/Runtime/Core/Private/ProfilingDebugging/TraceAuxiliary.cpp b/Engine/Source/Runtime/Core/Private/ProfilingDebugging/TraceAuxiliary.cpp
--- a/Engine/Source/Runtime/Core/Private/ProfilingDebugging/TraceAuxiliary.cpp
+++ b/Engine/Source/Runtime/Core/Private/ProfilingDebugging/TraceAuxiliary.cpp
@@ -36,6 +36,10 @@
 #include "Misc/Paths.h"
 #endif
 
+const FTraceAuxiliary::FChannelPreset GDefaultChannels(TEXT("Default"), TEXT("cpu,gpu,frame,log,bookmark,screenshot,region"), false);
+const FTraceAuxiliary::FChannelPreset GMemoryChannels(TEXT("Memory"), TEXT("memtag,memalloc,callstack,module"), true);
+const FTraceAuxiliary::FChannelPreset GMemoryLightChannels(TEXT("Memory_Light"), TEXT("memtag,memalloc"), true);
+
 #if UE_TRACE_ENABLED
 
 #include <atomic>
@@ -67,9 +71,7 @@
 #include "Trace/Trace.inl"
 
 ////////////////////////////////////////////////////////////////////////////////
-const TCHAR* GDefaultChannels = TEXT("cpu,gpu,frame,log,bookmark,screenshot,region");
-const TCHAR* GMemoryChannels = TEXT("memtag,memalloc,callstack,module");
-const TCHAR* GMemoryLightChannels = TEXT("memtag,memalloc");
+
 const TCHAR* GTraceConfigSection = TEXT("Trace.Config");
 static UE::Trace::FInitializeDesc GInitializeDesc;
 
@@ -277,17 +279,17 @@ void FTraceAuxiliaryImpl::ForEachChannel(const TCHAR* ChannelList, bool bResolve
 		{
 			FString Value;
 			// Check against hard coded presets
-			if (FCString::Stricmp(Name, TEXT("default")) == 0)
+			if (FCString::Stricmp(Name, GDefaultChannels.Name) == 0)
 			{
-				ForEachChannel(GDefaultChannels, false, LogCategory, Callable);
+				ForEachChannel(GDefaultChannels.Channels, false, LogCategory, Callable);
 			}
-			else if (FCString::Stricmp(Name, TEXT("memory")) == 0)
+			else if (FCString::Stricmp(Name,GMemoryChannels.Name) == 0)
 			{
-				ForEachChannel(GMemoryChannels, false, LogCategory, Callable);
+				ForEachChannel(GMemoryChannels.Channels, false, LogCategory, Callable);
 			}
-			else if (FCString::Stricmp(Name, TEXT("memory_light")) == 0)
+			else if (FCString::Stricmp(Name, GMemoryLightChannels.Name) == 0)
 			{
-				ForEachChannel(GMemoryLightChannels, false, LogCategory, Callable);
+				ForEachChannel(GMemoryLightChannels.Channels, false, LogCategory, Callable);
 			}
 			// Check against data driven presets (if available)
 			else if (GConfig && GConfig->GetString(TEXT("Trace.ChannelPresets"), Name, Value, GEngineIni))
@@ -1315,12 +1317,12 @@ static bool StartFromCommandlineArguments(const TCHAR* CommandLine, bool& bOutSt
 	}
 	else if (FParse::Param(CommandLine, TEXT("trace")))
 	{
-		Channels = GDefaultChannels;
+		Channels = GDefaultChannels.Channels;
 	}
 #if WITH_EDITOR
 	else
 	{
-		Channels = GDefaultChannels;
+		Channels = GDefaultChannels.Channels;
 	}
 #endif
 
@@ -1373,7 +1375,7 @@ static bool StartFromCommandlineArguments(const TCHAR* CommandLine, bool& bOutSt
 	// If user has defined a connection type but not specified channels, use the default channel set.
 	if (Type != FTraceAuxiliary::EConnectionType::None && Channels.IsEmpty())
 	{
-		Channels = GDefaultChannels;
+		Channels = GDefaultChannels.Channels;
 	}
 
 	if (Channels.IsEmpty())
@@ -1836,6 +1838,13 @@ UE::Trace::FInitializeDesc const* FTraceAuxiliary::GetInitializeDesc()
 #endif
 }
 
+void FTraceAuxiliary::GetFixedChannelPresets(TArray<FChannelPreset>& OutPresets)
+{
+	OutPresets.Emplace(GDefaultChannels);
+	OutPresets.Emplace(GMemoryChannels);
+	OutPresets.Emplace(GMemoryLightChannels);
+}
+
 ////////////////////////////////////////////////////////////////////////////////
 void FTraceAuxiliary::TryAutoConnect()
 {
diff --git a/Engine/Source/Runtime/Core/Public/ProfilingDebugging/TraceAuxiliary.h b/Engine/Source/Runtime/Core/Public/ProfilingDebugging/TraceAuxiliary.h
--- a/Engine/Source/Runtime/Core/Public/ProfilingDebugging/TraceAuxiliary.h
+++ b/Engine/Source/Runtime/Core/Public/ProfilingDebugging/TraceAuxiliary.h
@@ -19,6 +19,25 @@
 class FTraceAuxiliary
 {
 public:
+	struct FChannelPreset
+	{
+		FChannelPreset(const TCHAR* InName, const TCHAR* InChannels, bool bInIsReadOnly)
+			: Name(InName)
+			, Channels(InChannels)
+			, bIsReadOnly(bInIsReadOnly)
+		{
+		}
+
+		const TCHAR* Name;
+		const TCHAR* Channels;
+
+		/**
+		* A preset should be read-only if it contains any read-only channels. 
+		* A read-only preset can only be enabled using the command line when starting the application.
+		*/
+		bool bIsReadOnly = false;
+	};
+
 	// In no logging configurations all log categories are of type FNoLoggingCategory, which has no relation with
 	// FLogCategoryBase. In order to not need to conditionally set the argument alias the type here.
 #if NO_LOGGING
@@ -216,6 +235,11 @@ public:
 	 */
 	static CORE_API struct UE::Trace::FInitializeDesc const* GetInitializeDesc();
 
+	/**
+	* Get the channel presets that are defined in code.
+	*/
+	static CORE_API void GetFixedChannelPresets(TArray<FChannelPreset>& OutPresets);
+
 	/**
 	 * Delegate that triggers when a connection is established. Gives subscribers a chance to trace events that appear
 	 * after important events but before regular events (including tail). The following restrictions apply:
