ITraceController: added const FTraceStatus& parameter to FCallback typedef used by WithInstance and WithSelectedInstances pure virtual methods.

diff --git a/Engine/Source/Developer/TraceTools/Private/Services/SessionTraceControllerFilterService.cpp b/Engine/Source/Developer/TraceTools/Private/Services/SessionTraceControllerFilterService.cpp
--- a/Engine/Source/Developer/TraceTools/Private/Services/SessionTraceControllerFilterService.cpp
+++ b/Engine/Source/Developer/TraceTools/Private/Services/SessionTraceControllerFilterService.cpp
@@ -129,7 +129,7 @@ void FSessionTraceControllerFilterService::OnApplyChannelChanges()
 
 	if (FrameEnabledChannels.Num() || FrameDisabledChannels.Num())
 	{
-		TraceController->WithSelectedInstances([&](ITraceControllerCommands& Commands)
+		TraceController->WithSelectedInstances([&](const FTraceStatus& Status, ITraceControllerCommands& Commands)
 		{
 			Commands.SetChannels(FrameEnabledChannels.Array(), FrameDisabledChannels.Array());
 		});
diff --git a/Engine/Source/Developer/TraceTools/Private/Widgets/STraceControlToolbar.cpp b/Engine/Source/Developer/TraceTools/Private/Widgets/STraceControlToolbar.cpp
--- a/Engine/Source/Developer/TraceTools/Private/Widgets/STraceControlToolbar.cpp
+++ b/Engine/Source/Developer/TraceTools/Private/Widgets/STraceControlToolbar.cpp
@@ -175,7 +175,7 @@ bool STraceControlToolbar::StartTrace_CanExecute() const
 
 void STraceControlToolbar::StartTrace_Execute()
 {
-	TraceController->WithSelectedInstances([&](ITraceControllerCommands& Commands)
+	TraceController->WithSelectedInstances([&](const FTraceStatus& Status, ITraceControllerCommands& Commands)
 	{
 		if (TraceTarget == ETraceTarget::Server)
 		{
@@ -196,7 +196,7 @@ bool STraceControlToolbar::StopTrace_CanExecute() const
 
 void STraceControlToolbar::StopTrace_Execute()
 {
-	TraceController->WithSelectedInstances([&](ITraceControllerCommands& Commands)
+	TraceController->WithSelectedInstances([&](const FTraceStatus& Status, ITraceControllerCommands& Commands)
 	{
 		Commands.Stop();
 	});
@@ -210,7 +210,7 @@ bool STraceControlToolbar::TraceSnapshot_CanExecute() const
 
 void STraceControlToolbar::TraceSnapshot_Execute()
 {
-	TraceController->WithSelectedInstances([&](ITraceControllerCommands& Commands)
+	TraceController->WithSelectedInstances([&](const FTraceStatus& Status, ITraceControllerCommands& Commands)
 	{
 		if (TraceTarget == ETraceTarget::Server)
 		{
@@ -230,7 +230,7 @@ bool STraceControlToolbar::PauseTrace_CanExecute() const
 
 void STraceControlToolbar::PauseTrace_Execute()
 {
-	TraceController->WithSelectedInstances([&](ITraceControllerCommands& Commands)
+	TraceController->WithSelectedInstances([&](const FTraceStatus& Status, ITraceControllerCommands& Commands)
 	{
 		Commands.Pause();
 	});
@@ -244,7 +244,7 @@ bool STraceControlToolbar::ResumeTrace_CanExecute() const
 
 void STraceControlToolbar::ResumeTrace_Execute()
 {
-	TraceController->WithSelectedInstances([&](ITraceControllerCommands& Commands)
+	TraceController->WithSelectedInstances([&](const FTraceStatus& Status, ITraceControllerCommands& Commands)
 	{
 		Commands.Resume();
 	});
@@ -259,7 +259,7 @@ bool STraceControlToolbar::TraceBookmark_CanExecute() const
 void STraceControlToolbar::TraceBookmark_Execute()
 {
 	const FString BookmarkName = FDateTime::Now().ToString(TEXT("Bookmark_%Y%m%d_%H%M%S"));
-	TraceController->WithSelectedInstances([&](ITraceControllerCommands& Commands)
+	TraceController->WithSelectedInstances([&](const FTraceStatus& Status, ITraceControllerCommands& Commands)
 	{
 		Commands.Bookmark(BookmarkName);
 	});
@@ -272,7 +272,7 @@ bool STraceControlToolbar::TraceScreenshot_CanExecute() const
 
 void STraceControlToolbar::TraceScreenshot_Execute()
 {
-	TraceController->WithSelectedInstances([&](ITraceControllerCommands& Commands)
+	TraceController->WithSelectedInstances([&](const FTraceStatus& Status, ITraceControllerCommands& Commands)
 	{
 		Commands.Screenshot(TEXT(""), false);
 	});
@@ -291,7 +291,7 @@ bool STraceControlToolbar::ToggleStatNamedEvents_IsChecked() const
 void STraceControlToolbar::ToggleStatNamedEvents_Execute()
 {
 	bAreStatNamedEventsEnabled = !bAreStatNamedEventsEnabled;
-	TraceController->WithSelectedInstances([&](ITraceControllerCommands& Commands)
+	TraceController->WithSelectedInstances([&](const FTraceStatus& Status, ITraceControllerCommands& Commands)
 	{
 		Commands.SetStatNamedEventsEnabled(bAreStatNamedEventsEnabled);
 	});
diff --git a/Engine/Source/Runtime/SessionServices/Private/TraceController.cpp b/Engine/Source/Runtime/SessionServices/Private/TraceController.cpp
--- a/Engine/Source/Runtime/SessionServices/Private/TraceController.cpp
+++ b/Engine/Source/Runtime/SessionServices/Private/TraceController.cpp
@@ -279,7 +279,7 @@ void FTraceController::WithSelectedInstances(FCallback Func)
 		{
 			if(auto* Instance = Instances.Find(*InstanceAddress))
 			{
-				Func(Instance->Commands);
+				Func(Instance->Status, Instance->Commands);
 			}
 		}
 	}
@@ -292,7 +292,7 @@ void FTraceController::WithInstance(FGuid InstanceId, FCallback Func)
 	{
 		if(auto* Instance = Instances.Find(*InstanceAddress))
 		{
-			Func(Instance->Commands);
+			Func(Instance->Status, Instance->Commands);
 		}
 	}
 }
diff --git a/Engine/Source/Runtime/SessionServices/Public/ITraceController.h b/Engine/Source/Runtime/SessionServices/Public/ITraceController.h
--- a/Engine/Source/Runtime/SessionServices/Public/ITraceController.h
+++ b/Engine/Source/Runtime/SessionServices/Public/ITraceController.h
@@ -101,19 +101,19 @@ ENUM_CLASS_FLAGS(FTraceStatus::EUpdateType);
  *
  * Example usage:
  * ```
- * // I've found the session that we want to control
- * FGuid SessionId = ...;
+ * // I've found the session instance that we want to control
+ * FGuid InstanceId = ...;
  * // Get the controller
  * TSharedPtr<ITraceController> TraceController = SessionServicesModule.GetTraceController();
  * 
- * // To control just a single instance use the session id.
- * TraceController->WithInstance(SessionId, [](ITraceControllerCommands& Commands){
+ * // To control just a single instance use the session instance id.
+ * TraceController->WithInstance(InstanceId, [](const FTraceStatus& Status, ITraceControllerCommands& Commands){
  *		Commands.Send("localhost", "audio,audiomixer,bookmark,log");
  *		Commands.Bookmark("My remote bookmark");
  * });
  *
  * // We can also use the "selection" feature in session manager
- * TraceController->WithSelectedInstances([](ITraceControllerCommands& Commands){
+ * TraceController->WithSelectedInstances([](const FTraceStatus& Status, ITraceControllerCommands& Commands){
  *		Commands.SnapshotSend("localhost");
  * });
  * ```
@@ -165,7 +165,7 @@ public:
 	virtual bool HasAvailableSelectedInstance() = 0;
 
 
-	typedef TFunction<void(ITraceControllerCommands&)> FCallback;
+	typedef TFunction<void(const FTraceStatus&, ITraceControllerCommands&)> FCallback;
 	/**
 	 * Execute a function for each instance selected in the session manager.
 	 * @param Func Functor to execute
