Switched channel list to a horizontal layout and added trace settings and stats to the trace control window.

diff --git a/Engine/Source/Developer/TraceTools/Private/Models/ITraceObject.h b/Engine/Source/Developer/TraceTools/Private/Models/ITraceObject.h
--- a/Engine/Source/Developer/TraceTools/Private/Models/ITraceObject.h
+++ b/Engine/Source/Developer/TraceTools/Private/Models/ITraceObject.h
@@ -14,9 +14,16 @@ public:
 
 	/** Returns the name representing this object */
 	virtual FString GetName() const = 0;
-	/** Returns the user-facing text for this object */
+
+	/** Returns the description representing this object */
+	virtual FString GetDescription() const = 0;
+	
+	/** Returns the user-facing name for this object */
 	virtual FText GetDisplayText() const = 0;
 
+	/** Returns the tooltip text for this object */
+	virtual FText GetTooltipText() const = 0;
+
 	/** Set whether or not this object should (not) be filtered out */
 	virtual void SetIsFiltered(bool bState) = 0;
 
@@ -32,9 +39,6 @@ public:
 	/** Returns whether or not this object is pending a filter-state update (SetPending previously called) */
 	virtual bool IsPending() const = 0;
 
-	/** Returns contained child object(s) */
-	virtual void GetChildren(TArray<TSharedPtr<ITraceObject>>& OutChildren) const = 0;
-
 	/** Returns string representing this object in a search (treeview filtering) */
 	virtual void GetSearchString(TArray<FString>& OutFilterStrings) const = 0;
 };
diff --git a/Engine/Source/Developer/TraceTools/Private/Models/TraceChannel.cpp b/Engine/Source/Developer/TraceTools/Private/Models/TraceChannel.cpp
--- a/Engine/Source/Developer/TraceTools/Private/Models/TraceChannel.cpp
+++ b/Engine/Source/Developer/TraceTools/Private/Models/TraceChannel.cpp
@@ -10,9 +10,22 @@
 namespace UE::TraceTools
 {
 
-FTraceChannel::FTraceChannel(FString InName, FString InParentName, uint32 InHash, bool bInEnabled, bool bInReadOnly, const TArray<TSharedPtr<ITraceObject>>& InChildObjects, TSharedPtr<ISessionTraceFilterService> InFilterService) : Name(InName), ParentName(InParentName), Hash(InHash), ChildObjects(InChildObjects), bFiltered(!bInEnabled), bIsPending(false), bReadOnly(bInReadOnly), FilterService(InFilterService)
+FTraceChannel::FTraceChannel(FString InName,
+							 FString InDescription,
+							 FString InParentName, 
+							 uint32 InHash,
+							 bool bInEnabled,
+							 bool bInReadOnly,
+							 TSharedPtr<ISessionTraceFilterService> InFilterService) 
+	: Name(InName),
+	  Description(InDescription),
+	  ParentName(InParentName),
+	  Hash(InHash),
+	  bFiltered(!bInEnabled), 
+	  bIsPending(false),
+	  bReadOnly(bInReadOnly),
+	  FilterService(InFilterService)
 {
-
 }
 
 FText FTraceChannel::GetDisplayText() const
@@ -20,17 +33,26 @@ FText FTraceChannel::GetDisplayText() const
 	return FText::FromString(Name);
 }
 
+FText FTraceChannel::GetTooltipText() const
+{
+	return FText::FromString(Description);
+}
+
 FString FTraceChannel::GetName() const
 {
 	return Name;
 }
 
+FString FTraceChannel::GetDescription() const
+{
+	return Description;
+}
+
 void FTraceChannel::SetPending()
 {
 	bIsPending = true;
 }
 
-
 bool FTraceChannel::IsReadOnly() const
 {
 	return bReadOnly;
@@ -57,9 +79,4 @@ void FTraceChannel::GetSearchString(TArray<FString>& OutFilterStrings) const
 	OutFilterStrings.Add(Name);
 }
 
-void FTraceChannel::GetChildren(TArray<TSharedPtr<ITraceObject>>& OutChildren) const
-{
-	OutChildren.Append(ChildObjects);
-}
-
 } // namespace UE::TraceTools
\ No newline at end of file
diff --git a/Engine/Source/Developer/TraceTools/Private/Models/TraceChannel.h b/Engine/Source/Developer/TraceTools/Private/Models/TraceChannel.h
--- a/Engine/Source/Developer/TraceTools/Private/Models/TraceChannel.h
+++ b/Engine/Source/Developer/TraceTools/Private/Models/TraceChannel.h
@@ -12,28 +12,30 @@ class ISessionTraceFilterService;
 class FTraceChannel : public ITraceObject
 {
 public:
-	FTraceChannel(FString InName, FString InParentName, uint32 InHash, bool bInEnabled, bool bInReadOnly, const TArray<TSharedPtr<ITraceObject>>& InChildObjects, TSharedPtr<ISessionTraceFilterService> InFilterService);
+	FTraceChannel(FString InName, FString InDescription, FString InParentName, uint32 InHash, bool bInEnabled, bool bInReadOnly, TSharedPtr<ISessionTraceFilterService> InFilterService);
 
 	/** Begin ITraceObject overrides */
-	virtual FText GetDisplayText() const override;;
-	virtual FString GetName() const;
+	virtual FText GetDisplayText() const override;
+	virtual FText GetTooltipText() const override;
+	virtual FString GetName() const override;
+	virtual FString GetDescription() const override;
 	virtual void SetPending() override;
-	virtual bool IsReadOnly() const;
+	virtual bool IsReadOnly() const override;
 	virtual void SetIsFiltered(bool bState) override;
 	virtual bool IsFiltered() const override;
 	virtual bool IsPending() const override;
 	virtual void GetSearchString(TArray<FString>& OutFilterStrings) const override;
-	virtual void GetChildren(TArray<TSharedPtr<ITraceObject>>& OutChildren) const override;
 	/** End ITraceObject overrides */
 
 protected:
 	/** This channel's name */
 	FString Name;
+	/** This channel's description */
+	FString Description;
 	/** Channel's parent (group) name */
 	FString ParentName;
 	/** Hash representing this channel uniquely */
 	uint32 Hash;
-	TArray<TSharedPtr<ITraceObject>> ChildObjects;
 	
 	/** Whether or not this channel is filtered out, true = filtered; false = not filtered */
 	bool bFiltered;
diff --git a/Engine/Source/Developer/TraceTools/Private/Models/TraceFilterPresets.cpp b/Engine/Source/Developer/TraceTools/Private/Models/TraceFilterPresets.cpp
--- a/Engine/Source/Developer/TraceTools/Private/Models/TraceFilterPresets.cpp
+++ b/Engine/Source/Developer/TraceTools/Private/Models/TraceFilterPresets.cpp
@@ -133,13 +133,6 @@ void FFilterPresetHelpers::ExtractEnabledObjectNames(const TArray<TSharedPtr<ITr
 		{
 			OutNames.Add(Object->GetName());
 		}
-
-		TArray<TSharedPtr<ITraceObject>> ChildObjects;
-		Object->GetChildren(ChildObjects);
-		if(ChildObjects.Num())
-		{
-			ExtractEnabledObjectNames(ChildObjects, OutNames);
-		}
 	}
 }
 
diff --git a/Engine/Source/Developer/TraceTools/Private/Services/ISessionTraceFilterService.h b/Engine/Source/Developer/TraceTools/Private/Services/ISessionTraceFilterService.h
--- a/Engine/Source/Developer/TraceTools/Private/Services/ISessionTraceFilterService.h
+++ b/Engine/Source/Developer/TraceTools/Private/Services/ISessionTraceFilterService.h
@@ -4,6 +4,9 @@
 
 #include "Templates/SharedPointer.h"
 
+// TraceTools
+#include "ITraceController.h"
+
 struct FDateTime;
 
 namespace UE::TraceTools
@@ -14,10 +17,16 @@ struct ITraceFilterPreset;
 struct FTraceObjectInfo
 {
 	FString Name;
+	FString Description;
 	bool bEnabled;
 	bool bReadOnly;
 	uint32 Hash;
 	uint32 OwnerHash;
+
+	bool operator<(const FTraceObjectInfo& Other) const
+	{
+		return Name.Compare(Other.Name) < 0;
+	}
 };
 
 /** Filtering service, representing the state and data for a specific TraceServices::IAnalysisSession */
@@ -29,17 +38,26 @@ public:
 	/** Returns the root level set of objects */
 	virtual void GetRootObjects(TArray<FTraceObjectInfo>& OutObjects) const = 0;
 
-	/** Returns the contained objects for the specific object hash */
-	virtual void GetChildObjects(uint32 InObjectHash, TArray<FTraceObjectInfo>& OutChildObjects) const = 0;
-
 	/** Set the filtered state for an individual object by its hash */
 	virtual void SetObjectFilterState(const FString& InObjectName, const bool bFilterState) = 0;
 	
 	/** Set timestamp for last processed update (data change) */
-	virtual const FDateTime& GetTimestamp() = 0;
+	virtual const FDateTime& GetTimestamp() const = 0;
 
 	/** Update filter preset */
 	virtual void UpdateFilterPreset(const TSharedPtr<ITraceFilterPreset> InPreset, bool IsEnabled) = 0;
+
+	/** Returns true if settings are available for the selected session. */
+	virtual bool HasSettings() const = 0 ;
+
+	/** Get the settings of the selected session */
+	virtual const FTraceStatus::FSettings& GetSettings() const = 0;
+
+	/** Returns true if stats are available for the selected session. */
+	virtual bool HasStats() const = 0;
+
+	/** Get the settings of the selected session */
+	virtual const FTraceStatus::FStats& GetStats() const = 0;
 };
 
 } // namespace UE::TraceTools
\ No newline at end of file
diff --git a/Engine/Source/Developer/TraceTools/Private/Services/SessionTraceControllerFilterService.cpp b/Engine/Source/Developer/TraceTools/Private/Services/SessionTraceControllerFilterService.cpp
--- a/Engine/Source/Developer/TraceTools/Private/Services/SessionTraceControllerFilterService.cpp
+++ b/Engine/Source/Developer/TraceTools/Private/Services/SessionTraceControllerFilterService.cpp
@@ -28,12 +28,7 @@ void FSessionTraceControllerFilterService::GetRootObjects(TArray<FTraceObjectInf
 	OutObjects.Append(Objects);
 }
 
-void FSessionTraceControllerFilterService::GetChildObjects(uint32 InObjectHash, TArray<FTraceObjectInfo>& OutChildObjects) const
-{
-	/** TODO, parent/child relationship for Channels */
-}
-
-const FDateTime& FSessionTraceControllerFilterService::GetTimestamp()
+const FDateTime& FSessionTraceControllerFilterService::GetTimestamp() const
 {
 	return TimeStamp;
 }
@@ -92,13 +87,27 @@ void FSessionTraceControllerFilterService::OnTraceStatusUpdated(const FTraceStat
 		Objects.Empty();
 		return;
 	}
-	if (InUpdateType != FTraceStatus::EUpdateType::ChannelsDesc && 
-		InUpdateType != FTraceStatus::EUpdateType::ChannelsStatus && 
-		InUpdateType != FTraceStatus::EUpdateType::All)
+	if (EnumHasAnyFlags(InUpdateType, FTraceStatus::EUpdateType::ChannelsDesc) ||
+		(EnumHasAnyFlags(InUpdateType, FTraceStatus::EUpdateType::ChannelsStatus)))
 	{
-		return;
+		UpdateChannels(InStatus);
+	}
+
+	if (EnumHasAnyFlags(InUpdateType, FTraceStatus::EUpdateType::Settings))
+	{
+		Settings = InStatus.Settings;
+		bHasSettings = true;
 	}
 
+	if (EnumHasAnyFlags(InUpdateType, FTraceStatus::EUpdateType::Status))
+	{
+		Stats = InStatus.Stats;
+		bHasStats = true;
+	}
+}
+
+void FSessionTraceControllerFilterService::UpdateChannels(const FTraceStatus& InStatus)
+{
 	TimeStamp = FDateTime::Now();
 
 	const TMap<uint32, FTraceStatus::FChannel> Channels = InStatus.Channels;
@@ -113,6 +122,7 @@ void FSessionTraceControllerFilterService::OnTraceStatusUpdated(const FTraceStat
 	{
 		FTraceObjectInfo& EventInfo = Objects.AddDefaulted_GetRef();
 		EventInfo.Name = Entry.Value.Name;
+		EventInfo.Description = Entry.Value.Description;
 		EventInfo.bEnabled = Entry.Value.bEnabled;
 		EventInfo.bReadOnly = Entry.Value.bReadOnly;
 		EventInfo.Hash = Entry.Value.Id;
@@ -137,6 +147,26 @@ void FSessionTraceControllerFilterService::OnApplyChannelChanges()
 		FrameEnabledChannels.Empty();
 		FrameDisabledChannels.Empty();
 	}
-} 
+}
+
+bool FSessionTraceControllerFilterService::HasSettings() const
+{
+	return bHasSettings;
+}
+
+const FTraceStatus::FSettings& FSessionTraceControllerFilterService::GetSettings() const
+{
+	return Settings;
+}
+
+bool FSessionTraceControllerFilterService::HasStats() const
+{
+	return bHasStats;
+}
+
+const FTraceStatus::FStats& FSessionTraceControllerFilterService::GetStats() const
+{
+	return Stats;
+}
 
 } // namespace UE::TraceTools
\ No newline at end of file
diff --git a/Engine/Source/Developer/TraceTools/Private/Services/SessionTraceControllerFilterService.h b/Engine/Source/Developer/TraceTools/Private/Services/SessionTraceControllerFilterService.h
--- a/Engine/Source/Developer/TraceTools/Private/Services/SessionTraceControllerFilterService.h
+++ b/Engine/Source/Developer/TraceTools/Private/Services/SessionTraceControllerFilterService.h
@@ -3,7 +3,6 @@
 #pragma once
 
 #include "ISessionTraceFilterService.h"
-#include "ITraceController.h"
 #include "Misc/DateTime.h"
 
 namespace TraceServices
@@ -26,10 +25,15 @@ public:
 
 	/** Begin ISessionTraceFilterService overrides */
 	virtual void GetRootObjects(TArray<FTraceObjectInfo>& OutObjects) const override;
-	virtual void GetChildObjects(uint32 InObjectHash, TArray<FTraceObjectInfo>& OutChildObjects) const override;
-	virtual const FDateTime& GetTimestamp() override;
+	virtual const FDateTime& GetTimestamp() const override;
 	virtual void SetObjectFilterState(const FString& InObjectName, const bool bFilterState) override;
 	virtual void UpdateFilterPreset(const TSharedPtr<ITraceFilterPreset> InPreset, bool IsEnabled) override;
+
+	virtual bool HasSettings() const override;
+	virtual const FTraceStatus::FSettings& GetSettings() const override;
+
+	virtual bool HasStats() const override;
+	virtual const FTraceStatus::FStats& GetStats() const override;
 	/** End ISessionTraceFilterService overrides */
 
 protected:
@@ -43,6 +47,8 @@ protected:
 
 	void OnTraceStatusUpdated(const FTraceStatus& InStatus, FTraceStatus::EUpdateType InUpdateType, ITraceControllerCommands& Commands);
 
+	void UpdateChannels(const FTraceStatus& InStatus);
+
 protected:
 	TSharedPtr<ITraceController> TraceController;
 
@@ -56,6 +62,12 @@ protected:
 	FDateTime TimeStamp;
 
 	bool bChannelsReceived = false;
+
+	FTraceStatus::FSettings Settings;
+	bool bHasSettings = false;
+
+	FTraceStatus::FStats Stats;
+	bool bHasStats = false;
 };
 
 } // namespace UE::TraceTools
\ No newline at end of file
diff --git a/Engine/Source/Developer/TraceTools/Private/TraceToolsModule.cpp b/Engine/Source/Developer/TraceTools/Private/TraceToolsModule.cpp
--- a/Engine/Source/Developer/TraceTools/Private/TraceToolsModule.cpp
+++ b/Engine/Source/Developer/TraceTools/Private/TraceToolsModule.cpp
@@ -20,7 +20,7 @@ void FTraceToolsModule::StartupModule()
 
 	UE::TraceTools::FTraceToolsStyle::Initialize();
 
-	FConfigContext::ReadIntoGConfig().Load(TEXT("Trace"), TraceFiltersIni);
+	TraceFiltersIni = GEngineIni;
 }
 
 void FTraceToolsModule::ShutdownModule()
diff --git a/Engine/Source/Developer/TraceTools/Private/TraceToolsStyle.cpp b/Engine/Source/Developer/TraceTools/Private/TraceToolsStyle.cpp
--- a/Engine/Source/Developer/TraceTools/Private/TraceToolsStyle.cpp
+++ b/Engine/Source/Developer/TraceTools/Private/TraceToolsStyle.cpp
@@ -133,7 +133,9 @@ void FTraceToolsStyle::Initialize()
 	StyleSet->Set("FilterPresets.FilterNameFont", DEFAULT_FONT("Regular", 10));
 	StyleSet->Set("FilterPresets.FilterButtonBorder", new BOX_BRUSH("Common/RoundedSelection_16x", FMargin(4.0f / 16.0f)));
 
-	StyleSet->Set("FilterPresets.SessionWarningBorder", new BOX_BRUSH("Common/DarkGroupBorder", FMargin(4.0f / 16.0f)));
+	StyleSet->Set("FilterPresets.TableBackground", new BOX_BRUSH("Common/TableViewMajorColumn", FMargin(4.0f / 16.0f)));
+	StyleSet->Set("FilterPresets.SessionWarningBorder", new BOX_BRUSH("Common/GroupBorder", FMargin(4.0f / 16.0f)));
+	StyleSet->Set("FilterPresets.BackgroundBorder", new BOX_BRUSH("Common/DarkGroupBorder", FMargin(4.0f / 16.0f)));
 	StyleSet->Set("FilterPresets.WarningIcon", new CORE_IMAGE_BRUSH_SVG("Starship/Common/alert-circle", Icon40x40, FStyleColors::Warning));
 
 	StyleSet->Set("FontAwesome.9", ICON_FONT(9));
diff --git a/Engine/Source/Developer/TraceTools/Private/Widgets/STraceControl.cpp b/Engine/Source/Developer/TraceTools/Private/Widgets/STraceControl.cpp
--- a/Engine/Source/Developer/TraceTools/Private/Widgets/STraceControl.cpp
+++ b/Engine/Source/Developer/TraceTools/Private/Widgets/STraceControl.cpp
@@ -47,21 +47,6 @@ void STraceControl::Construct(const FArguments& InArgs, TSharedPtr<ITraceControl
 	];
 }
 
-void STraceControl::Tick(const FGeometry& AllottedGeometry, const double InCurrentTime, const float InDeltaTime)
-{
-	AccumulatedTime += InDeltaTime;
-
-	const double UpdateTime = 1.0f;
-
-	if (AccumulatedTime > UpdateTime)
-	{
-		TraceController->SendStatusUpdateRequest();
-		TraceController->SendChannelUpdateRequest();
-
-		AccumulatedTime = 0.0;
-	}
-}
-
 void STraceControl::BindCommands()
 {
 
diff --git a/Engine/Source/Developer/TraceTools/Private/Widgets/STraceControl.h b/Engine/Source/Developer/TraceTools/Private/Widgets/STraceControl.h
--- a/Engine/Source/Developer/TraceTools/Private/Widgets/STraceControl.h
+++ b/Engine/Source/Developer/TraceTools/Private/Widgets/STraceControl.h
@@ -26,8 +26,6 @@ public:
 	/** Constructs this widget. */
 	void Construct(const FArguments& InArgs, TSharedPtr<ITraceController> InTraceController);
 
-	virtual void Tick(const FGeometry& AllottedGeometry, const double InCurrentTime, const float InDeltaTime) override;
-
 protected:
 	void BindCommands();
 
@@ -35,8 +33,6 @@ private:
 	TSharedPtr<ITraceController> TraceController;
 
 	TSharedPtr<FUICommandList> UICommandList;
-
-	double AccumulatedTime = 0.0f;
 };
 
 } // namespace UE::TraceTools
diff --git a/Engine/Source/Developer/TraceTools/Private/Widgets/STraceDataFilterWidget.cpp b/Engine/Source/Developer/TraceTools/Private/Widgets/STraceDataFilterWidget.cpp
--- a/Engine/Source/Developer/TraceTools/Private/Widgets/STraceDataFilterWidget.cpp
+++ b/Engine/Source/Developer/TraceTools/Private/Widgets/STraceDataFilterWidget.cpp
@@ -8,6 +8,7 @@
 #include "Misc/TextFilter.h"
 #include "Modules/ModuleManager.h"
 #include "Widgets/Layout/SScrollBorder.h"
+#include "Widgets/Layout/SSeparator.h"
 #include "Widgets/Input/SComboButton.h"
 #include "Widgets/Layout/SScrollBar.h"
 #include "Widgets/Input/SSearchBox.h"
@@ -18,6 +19,7 @@
 #include "Models/TraceFilterPresets.h"
 #include "Services/SessionTraceControllerFilterService.h"
 #include "STraceObjectRowWidget.h"
+#include "STraceStatistics.h"
 #include "TraceToolsStyle.h"
 #include "Widgets/SFilterPresetList.h"
 
@@ -26,18 +28,12 @@
 namespace UE::TraceTools
 {
 
-TSet<FString> STraceDataFilterWidget::LastExpandedObjectNames;
-
-STraceDataFilterWidget::STraceDataFilterWidget() : PreviousSessionHandle(INDEX_NONE), bNeedsTreeRefresh(false), bHighlightingPreset(false)
+STraceDataFilterWidget::STraceDataFilterWidget() : bNeedsListRefresh(false), bHighlightingPreset(false)
 {
 }
 
 STraceDataFilterWidget::~STraceDataFilterWidget()
 {
-	// Save expansion and move to static array
-	SaveItemsExpansion();
-	LastExpandedObjectNames = ExpandedObjectNames;
-	ExpandedObjectNames.Empty();
 }
 
 void STraceDataFilterWidget::Construct(const FArguments& InArgs, TSharedPtr<ITraceController> InTraceController)
@@ -49,14 +45,13 @@ void STraceDataFilterWidget::Construct(const FArguments& InArgs, TSharedPtr<ITra
 	.AlwaysShowScrollbar(true);
 
 	ConstructSearchBoxFilter();
-	ConstructFilterHandler();
-	ConstructTreeview();
+	ConstructTileView();
 
 	ChildSlot
 	[		
 		SNew(SBorder)
 		.Padding(4)
-		.BorderImage(FTraceToolsStyle::GetBrush("EventFilter.GroupBorder"))
+		.BorderImage(FTraceToolsStyle::GetBrush("FilterPresets.BackgroundBorder"))
 		[
 			SNew( SVerticalBox )
 
@@ -169,23 +164,37 @@ void STraceDataFilterWidget::Construct(const FArguments& InArgs, TSharedPtr<ITra
 				.HAlign(HAlign_Fill)
 				.VAlign(VAlign_Fill)
 				[
-					SNew(SHorizontalBox)
+					SNew(SVerticalBox)
 
-					+ SHorizontalBox::Slot()
+					+ SVerticalBox::Slot()
 					[
-						Treeview.ToSharedRef()
+						SNew(SBorder)
+						.BorderImage(FTraceToolsStyle::GetBrush("FilterPresets.TableBackground"))
+						[
+							TileView.ToSharedRef()
+						]
 					]
 
-					+ SHorizontalBox::Slot()
-					.AutoWidth()
+					+ SVerticalBox::Slot()
+					.AutoHeight()
+					.Padding(0.0f, 8.0f, 0.0f, 0.0f)
+
 					[
-						SNew(SBox)
-						.WidthOverride(FOptionalSize(13.0f))
-						[
-							ExternalScrollbar.ToSharedRef()
-						]
+						SNew(SSeparator)
+						.Thickness(5.0f)
+						.SeparatorImage(FTraceToolsStyle::GetBrush("FilterPresets.TableBackground"))
+						.Orientation(EOrientation::Orient_Horizontal)
 					]
-				]		
+
+					+ SVerticalBox::Slot()
+					[
+						SNew(STraceStatistics, SessionFilterService)
+						.Visibility_Lambda([this]() -> EVisibility
+						{
+							return HasValidFilterSession() ? EVisibility::Visible : EVisibility::Collapsed;
+						})
+					]
+				]
 			]
 		]
 	];
@@ -194,13 +203,9 @@ void STraceDataFilterWidget::Construct(const FArguments& InArgs, TSharedPtr<ITra
 	TAttribute<bool> EnabledAttribute;
 	EnabledAttribute.Bind(this, &STraceDataFilterWidget::HasValidFilterSession);
 
-	Treeview->SetEnabled(EnabledAttribute);
+	TileView->SetEnabled(EnabledAttribute);
 	OptionsWidget->SetEnabled(EnabledAttribute);
 	FilterPresetsListWidget->SetEnabled(EnabledAttribute);
-
-	/** Restore expansion state when recreating this window */
-	ExpandedObjectNames = LastExpandedObjectNames;
-	RestoreItemsExpansion();
 }
 
 void STraceDataFilterWidget::OnSavePreset(const TSharedPtr<ITraceFilterPreset>& Preset)
@@ -210,11 +215,11 @@ void STraceDataFilterWidget::OnSavePreset(const TSharedPtr<ITraceFilterPreset>&
 		/** Save to preset if one was provided, otherwise create a new one */
 		if (Preset.IsValid())
 		{
-			Preset->Save(RootItems);
+			Preset->Save(ListItems);
 		}
 		else
 		{
-			FFilterPresetHelpers::CreateNewPreset(RootItems);
+			FFilterPresetHelpers::CreateNewPreset(ListItems);
 		}
 	}
 }
@@ -229,16 +234,15 @@ void STraceDataFilterWidget::OnPresetChanged(const SFilterPreset& Preset)
 
 void STraceDataFilterWidget::OnHighlightPreset(const TSharedPtr<ITraceFilterPreset>& Preset)
 {
-	if (Treeview.IsValid())
+	if (TileView.IsValid())
 	{
-		/** Update treeview so that any allowlisted entry (as part of Preset) is highlighted and expanded */
-		Treeview->ClearHighlightedItems();
+		/** Update tileview so that any allowlisted entry (as part of Preset) is highlighted */
+		TileView->ClearHighlightedItems();
 		if (Preset.IsValid())
 		{
 			if (!bHighlightingPreset)
 			{
 				/** Store current expansion, so we can reset once highlighting has finished */
-				SaveItemsExpansion();
 				bHighlightingPreset = true;
 			}
 
@@ -249,68 +253,32 @@ void STraceDataFilterWidget::OnHighlightPreset(const TSharedPtr<ITraceFilterPres
 			{
 				if (Names.Contains(Object->GetName()))
 				{
-					Treeview->SetItemHighlighted(Object, true);
-					// TODO need to set parent expansion as well
-					Treeview->SetItemExpansion(Object, true);
-
-					SetParentExpansionRecursively(Object, true);
+					TileView->SetItemHighlighted(Object, true);
 				}
 			});
 		}
 		else
 		{
 			bHighlightingPreset = false;
-			RestoreItemsExpansion();
 		}
 	}
 }
 
 void STraceDataFilterWidget::OnSearchboxTextChanged(const FText& FilterText)
 {
-	bNeedsTreeRefresh = true;
-	const bool bFilterSet = !FilterText.IsEmpty();
-	if (bFilterSet != TreeviewFilterHandler->GetIsEnabled())
-	{
-		TreeviewFilterHandler->SetIsEnabled(bFilterSet);
-
-		if (bFilterSet)
-		{
-			SaveItemsExpansion();
-		}
-		else
-		{
-			RestoreItemsExpansion();
-		}
-	}
+	bNeedsListRefresh = true;
 
 	SearchBoxWidgetFilter->SetRawFilterText(FilterText);
-	SearchBoxWidget->SetError(SearchBoxWidgetFilter->GetFilterErrorText());
 }
 
-void STraceDataFilterWidget::ConstructTreeview()
+void STraceDataFilterWidget::ConstructTileView()
 {
-	SAssignNew(Treeview, STreeView<TSharedPtr<ITraceObject>>)
-	.OnGetChildren(TreeviewFilterHandler.ToSharedRef(), &TreeFilterHandler<TSharedPtr<ITraceObject>>::OnGetFilteredChildren)
-	.OnGenerateRow(this, &STraceDataFilterWidget::OnGenerateRow)
+	SAssignNew(TileView, STileView<TSharedPtr<ITraceObject>>)
+	.OnGenerateTile(this, &STraceDataFilterWidget::OnGenerateRow)
 	.OnContextMenuOpening(this, &STraceDataFilterWidget::OnContextMenuOpening)
-	.OnMouseButtonDoubleClick(this, &STraceDataFilterWidget::OnItemDoubleClicked)
-	.ExternalScrollbar(ExternalScrollbar)
-	.OnSetExpansionRecursive_Lambda([this](TSharedPtr<ITraceObject> InObject, bool bInExpansionState)
-	{
-		SetExpansionRecursively(InObject, bInExpansionState);
-	})
-	.TreeItemsSource(&TreeItems);
-	
-	TreeviewFilterHandler->SetTreeView(Treeview.Get());
-}
-
-void STraceDataFilterWidget::OnItemDoubleClicked(TSharedPtr<ITraceObject> InObject) const
-{
-	TSharedPtr<ITableRow> Row = Treeview->WidgetFromItem(InObject);
-	if (Row.IsValid())
-	{
-		SetExpansionRecursively(InObject, !Row->IsItemExpanded());
-	}
+	.Orientation(EOrientation::Orient_Horizontal)
+	.ItemHeight(15.0f)
+	.ListItemsSource(&FilteredListItems);
 }
 
 void STraceDataFilterWidget::ConstructSearchBoxFilter()
@@ -321,17 +289,6 @@ void STraceDataFilterWidget::ConstructSearchBoxFilter()
 	})));
 }
 
-void STraceDataFilterWidget::ConstructFilterHandler()
-{
-	TreeviewFilterHandler = MakeShareable(new TreeFilterHandler<TSharedPtr<ITraceObject>>());
-	TreeviewFilterHandler->SetFilter(SearchBoxWidgetFilter.Get());
-	TreeviewFilterHandler->SetRootItems(&RootItems, &TreeItems);
-	TreeviewFilterHandler->SetGetChildrenDelegate(TreeFilterHandler<TSharedPtr<ITraceObject>>::FOnGetChildren::CreateLambda([](TSharedPtr<ITraceObject> InParent, TArray<TSharedPtr<ITraceObject>>& OutChildren)
-	{
-		InParent->GetChildren(OutChildren);
-	}));
-}
-
 TSharedRef<ITableRow> STraceDataFilterWidget::OnGenerateRow(TSharedPtr<ITraceObject> InItem, const TSharedRef<STableViewBase>& OwnerTable)
 {
 	LLM_SCOPE_BYNAME(TEXT("Insights/STraceDataFilterWidget"));
@@ -363,7 +320,7 @@ TSharedPtr<SWidget> STraceDataFilterWidget::OnContextMenuOpening() const
 				FGetActionCheckState(),
 				FIsActionButtonVisible::CreateLambda([this]() -> bool
 				{
-					return Treeview->GetNumItemsSelected() == 0 && EnumerateAllItems(TFunction<bool(TSharedPtr<ITraceObject> InItem)>([this](TSharedPtr<ITraceObject> Object) -> bool
+					return TileView->GetNumItemsSelected() == 0 && EnumerateAllItems(TFunction<bool(TSharedPtr<ITraceObject> InItem)>([this](TSharedPtr<ITraceObject> Object) -> bool
 					{
 						return Object->IsFiltered();
 					}));
@@ -384,7 +341,7 @@ TSharedPtr<SWidget> STraceDataFilterWidget::OnContextMenuOpening() const
 				FGetActionCheckState(),
 				FIsActionButtonVisible::CreateLambda([this]() -> bool
 				{
-					return Treeview->GetNumItemsSelected() == 0 && EnumerateAllItems(TFunction<bool(TSharedPtr<ITraceObject> InItem)>([this](TSharedPtr<ITraceObject> Object) -> bool
+					return TileView->GetNumItemsSelected() == 0 && EnumerateAllItems(TFunction<bool(TSharedPtr<ITraceObject> InItem)>([this](TSharedPtr<ITraceObject> Object) -> bool
 					{
 						return !Object->IsFiltered();
 					}));
@@ -398,8 +355,6 @@ TSharedPtr<SWidget> STraceDataFilterWidget::OnContextMenuOpening() const
 				FExecuteAction::CreateRaw(this, &STraceDataFilterWidget::EnumerateSelectedItems, TFunction<void(TSharedPtr<ITraceObject> InItem)>([this](TSharedPtr<ITraceObject> Object) -> void
 				{
 					Object->SetIsFiltered(false);
-
-					SetExpansionRecursively(Object, true);
 				})),
 				FCanExecuteAction(),
 				FGetActionCheckState(),
@@ -433,126 +388,14 @@ TSharedPtr<SWidget> STraceDataFilterWidget::OnContextMenuOpening() const
 	}
 	MenuBuilder.EndSection();
 
-	static const FName ExpansionSectionHook("ExpansionState");
-	MenuBuilder.BeginSection(ExpansionSectionHook, LOCTEXT("ExpansionSectionLabel", "Expansion"));
-	{
-		MenuBuilder.AddMenuEntry(LOCTEXT("ExpandAllRowsLabel", "Expand All"), LOCTEXT("ExpandAllRowsTooltip", "Expands the entire hierarchy."), FSlateIcon(), FUIAction(
-				FExecuteAction::CreateLambda([this]()
-				{
-					for (const TSharedPtr<ITraceObject>& Object : TreeItems)
-					{
-						SetExpansionRecursively(Object, true);
-					}
-				}),
-				FCanExecuteAction(),
-				FGetActionCheckState(),
-				FIsActionButtonVisible::CreateLambda([this]() -> bool
-				{
-					return Treeview->GetNumItemsSelected() == 0 && EnumerateAllItems(TFunction<bool(TSharedPtr<ITraceObject> InItem)>([this](TSharedPtr<ITraceObject> Object) -> bool
-					{
-						return !Treeview->IsItemExpanded(Object);
-					}));
-				})
-			)
-		);
-		
-		MenuBuilder.AddMenuEntry(LOCTEXT("CollapseAllRowsLabel", "Collapse All"), LOCTEXT("CollapseAllRowsTooltip", "Collapses the entire hierarchy."), FSlateIcon(), 
-			FUIAction(
-				FExecuteAction::CreateLambda([this]()
-				{
-					for (const TSharedPtr<ITraceObject>& Object : TreeItems)
-					{
-						SetExpansionRecursively(Object, false);
-					}
-				}),
-				FCanExecuteAction(),
-				FGetActionCheckState(),
-				FIsActionButtonVisible::CreateLambda([this]() -> bool
-				{
-					return Treeview->GetNumItemsSelected() == 0 && EnumerateAllItems(TFunction<bool(TSharedPtr<ITraceObject> InItem)>([this](TSharedPtr<ITraceObject> Object) -> bool
-					{
-						return Treeview->IsItemExpanded(Object);
-					}));
-				})
-			)
-		);
-
-		MenuBuilder.AddMenuEntry(LOCTEXT("ExpandRowsLabel", "Expand Selected"), LOCTEXT("ExpandRowsTooltip", "Expands the selected Node(s)."), FSlateIcon(), 
-			FUIAction(
-				FExecuteAction::CreateRaw(this, &STraceDataFilterWidget::EnumerateSelectedItems, TFunction<void(TSharedPtr<ITraceObject> InItem)>([this](TSharedPtr<ITraceObject> Object) -> void
-				{
-					SetExpansionRecursively(Object, true);
-				})),
-				FCanExecuteAction(),
-				FGetActionCheckState(),
-				FIsActionButtonVisible::CreateLambda([this]() -> bool
-				{
-					return EnumerateSelectedItems(TFunction<bool(TSharedPtr<ITraceObject> InItem)>([this](TSharedPtr<ITraceObject> Object) -> bool
-					{
-						return !Treeview->IsItemExpanded(Object);
-					}));
-				})
-			)
-		);
-
-		MenuBuilder.AddMenuEntry(LOCTEXT("CollapseRowsLabel", "Collapse Selected"), LOCTEXT("CollapseRowsTooltip", "Collapse the selected Node(s)."), FSlateIcon(), 
-			FUIAction(
-				FExecuteAction::CreateRaw(this, &STraceDataFilterWidget::EnumerateSelectedItems, TFunction<void(TSharedPtr<ITraceObject> InItem)>([this](TSharedPtr<ITraceObject> Object) -> void
-				{
-					SetExpansionRecursively(Object, false);
-				})),
-				FCanExecuteAction(),
-				FGetActionCheckState(),
-				FIsActionButtonVisible::CreateLambda([this]() -> bool
-				{
-					return EnumerateSelectedItems(TFunction<bool(TSharedPtr<ITraceObject> InItem)>([this](TSharedPtr<ITraceObject> Object) -> bool
-					{
-						return Treeview->IsItemExpanded(Object);
-					}));
-				})
-			)
-		);
-	}  	
-	MenuBuilder.EndSection();
-
 	return MenuBuilder.MakeWidget();
 }
 
-void STraceDataFilterWidget::SaveItemsExpansion()
-{
-	ExpandedObjectNames.Empty();
-
-	if (Treeview.IsValid())
-	{
-		TSet<TSharedPtr<ITraceObject>> ExpandedItems;
-		Treeview->GetExpandedItems(ExpandedItems);
-
-		for (TSharedPtr<ITraceObject> Item : ExpandedItems)
-		{
-			if (Item.IsValid())
-			{
-				ExpandedObjectNames.Add(Item->GetName());
-			}
-		}
-	}
-}
-
-void STraceDataFilterWidget::RestoreItemsExpansion()
-{
-	EnumerateAllItems([this](TSharedPtr<ITraceObject> Object) -> void
-	{
-		const bool bExpanded = ExpandedObjectNames.Contains(Object->GetName());
-		SetExpansionRecursively(Object, bExpanded);
-	});
-
-	ExpandedObjectNames.Empty();
-}
-
 void STraceDataFilterWidget::SaveItemSelection()
 {
 	SelectedObjectNames.Empty();
 
-	if (Treeview.IsValid())
+	if (TileView.IsValid())
 	{
 		EnumerateSelectedItems([this](TSharedPtr<ITraceObject> InObject)
 		{
@@ -566,7 +409,7 @@ void STraceDataFilterWidget::SaveItemSelection()
 
 void STraceDataFilterWidget::RestoreItemSelection()
 {
-	if (Treeview.IsValid())
+	if (TileView.IsValid())
 	{
 		TArray<TSharedPtr<ITraceObject>> SelectedItems;
 		EnumerateAllItems([this, &SelectedItems](TSharedPtr<ITraceObject> Object) -> void
@@ -577,7 +420,7 @@ void STraceDataFilterWidget::RestoreItemSelection()
 			}
 		});
 
-		Treeview->SetItemSelection(SelectedItems, true);
+		TileView->SetItemSelection(SelectedItems, true);
 	}
 	
 	SelectedObjectNames.Empty();
@@ -585,26 +428,9 @@ void STraceDataFilterWidget::RestoreItemSelection()
 
 TSharedRef<ITraceObject> STraceDataFilterWidget::AddFilterableObject(const FTraceObjectInfo& Event, FString ParentName)
 {
-	// Retrieve any child objects, and recursively add those first 
-	TArray<FTraceObjectInfo> Events;
-	SessionFilterService->GetChildObjects(Event.Hash, Events);
+	TSharedRef<FTraceChannel> SharedItem = MakeShareable(new FTraceChannel(Event.Name, Event.Description, ParentName, Event.Hash, Event.bEnabled, Event.bReadOnly, SessionFilterService));
 
-	TArray<TSharedPtr<ITraceObject>> ChildPtrs;	
-	for (const FTraceObjectInfo& ChildEvent : Events)
-	{
-		TSharedRef<ITraceObject> EventItem = AddFilterableObject(ChildEvent, Event.Name);
-		ChildPtrs.Add(EventItem);
-	}
-
-	TSharedRef<FTraceChannel> SharedItem = MakeShareable(new FTraceChannel(Event.Name, ParentName, Event.Hash, Event.bEnabled, Event.bReadOnly, ChildPtrs, SessionFilterService));
-	ParentToChild.Add(SharedItem, ChildPtrs);
-
-	for (TSharedPtr<ITraceObject> ChildObject : ChildPtrs)
-	{
-		ChildToParent.Add(ChildObject, SharedItem);
-	}
-
-	FlatItems.Add(SharedItem);
+	ListItems.Add(SharedItem);
 
 	return SharedItem;
 }
@@ -614,32 +440,10 @@ bool STraceDataFilterWidget::HasValidFilterSession() const
 	return TraceController->HasAvailableSelectedInstance() && SessionFilterService.IsValid();
 }
 
-void STraceDataFilterWidget::SetExpansionRecursively(const TSharedPtr<ITraceObject>& InObject, bool bShouldExpandItem) const
-{
-	Treeview->SetItemExpansion(InObject, bShouldExpandItem);
-
-	TArray<TSharedPtr<ITraceObject>> Children;
-	InObject->GetChildren(Children);
-
-	for (TSharedPtr<ITraceObject>& ChildObject : Children)
-	{
-		SetExpansionRecursively(ChildObject, bShouldExpandItem);
-	}
-}
-
-void STraceDataFilterWidget::SetParentExpansionRecursively(const TSharedPtr<ITraceObject>& InObject, bool bShouldExpandItem) const
-{
-	if (const TSharedPtr<ITraceObject>* ParentObject = ChildToParent.Find(InObject))
-	{
-		Treeview->SetItemExpansion(*ParentObject, bShouldExpandItem);
-		SetParentExpansionRecursively(*ParentObject, bShouldExpandItem);
-	}
-}
-
 void STraceDataFilterWidget::EnumerateSelectedItems(TFunction<void(TSharedPtr<ITraceObject> InItem)> InFunction) const
 {
 	TArray<TSharedPtr<ITraceObject>> SelectedObjects;
-	Treeview->GetSelectedItems(SelectedObjects);
+	TileView->GetSelectedItems(SelectedObjects);
 
 	for (const TSharedPtr<ITraceObject>& Object : SelectedObjects)
 	{
@@ -650,7 +454,7 @@ void STraceDataFilterWidget::EnumerateSelectedItems(TFunction<void(TSharedPtr<IT
 bool STraceDataFilterWidget::EnumerateSelectedItems(TFunction<bool(TSharedPtr<ITraceObject> InItem)> InFunction) const
 {
 	TArray<TSharedPtr<ITraceObject>> SelectedObjects;
-	Treeview->GetSelectedItems(SelectedObjects);
+	TileView->GetSelectedItems(SelectedObjects);
 
 
 	bool bState = false;
@@ -664,7 +468,7 @@ bool STraceDataFilterWidget::EnumerateSelectedItems(TFunction<bool(TSharedPtr<IT
 
 void STraceDataFilterWidget::EnumerateAllItems(TFunction<void(TSharedPtr<ITraceObject> InItem)> InFunction) const
 {
-	for (const TSharedPtr<ITraceObject>& Object : FlatItems)
+	for (const TSharedPtr<ITraceObject>& Object : ListItems)
 	{
 		InFunction(Object);
 	}
@@ -673,39 +477,34 @@ void STraceDataFilterWidget::EnumerateAllItems(TFunction<void(TSharedPtr<ITraceO
 bool STraceDataFilterWidget::EnumerateAllItems(TFunction<bool(TSharedPtr<ITraceObject> InItem)> InFunction) const
 {
 	bool bState = false;
-	for (const TSharedPtr<ITraceObject>& Object : FlatItems)
+	for (const TSharedPtr<ITraceObject>& Object : ListItems)
 	{
 		bState |= InFunction(Object);
 	}
 	return bState;
 }
 
-void STraceDataFilterWidget::RefreshTreeviewData()
+void STraceDataFilterWidget::RefreshTileViewData()
 {
 	if (SessionFilterService.IsValid())
 	{
 		SyncTimeStamp = SessionFilterService->GetTimestamp();
 
 		/** Save expansion and selection */
-		SaveItemsExpansion();
 		SaveItemSelection();
 
+		ListItems.Empty();
+
 		TArray<FTraceObjectInfo> RootEvents;
 		SessionFilterService->GetRootObjects(RootEvents);
-		
-		ParentToChild.Empty();
-		ChildToParent.Empty();
-		RootItems.Empty();		
 
+		RootEvents.Sort();
+		
 		for (const FTraceObjectInfo& RootEvent : RootEvents)
 		{
 			TSharedRef<ITraceObject> TraceObject = AddFilterableObject(RootEvent, FString());
-			RootItems.Add(TraceObject);
 		}
 
-		TreeviewFilterHandler->RefreshAndFilterTree();
-
-		RestoreItemsExpansion();
 		RestoreItemSelection();
 	}
 }
@@ -718,13 +517,41 @@ void STraceDataFilterWidget::Tick(const FGeometry& AllottedGeometry, const doubl
 	{
 		if (SessionFilterService->GetTimestamp() != SyncTimeStamp)
 		{
-			RefreshTreeviewData();
+			RefreshTileViewData();
+			bNeedsListRefresh = true;
+		}
+
+		if (bNeedsListRefresh)
+		{
+			FilteredListItems.Empty();
+			for (auto& Item : ListItems)
+			{
+				if (SearchBoxWidgetFilter->PassesFilter(Item))
+				{
+					FilteredListItems.Add(Item);
+				}
+			}
+
+			TileView->RequestListRefresh();
+			bNeedsListRefresh = false;
 		}
 	}
 
-	if (bNeedsTreeRefresh)
+	AccumulatedTime += InDeltaTime;
+
+	constexpr double UpdateTime = 1.0f;
+
+	if (AccumulatedTime > UpdateTime)
 	{
-		TreeviewFilterHandler->RefreshAndFilterTree();
+		TraceController->SendStatusUpdateRequest();
+		TraceController->SendChannelUpdateRequest();
+
+		if (!SessionFilterService->HasSettings())
+		{
+			TraceController->SendSettingsUpdateRequest();
+		}
+
+		AccumulatedTime = 0.0;
 	}
 }
 
diff --git a/Engine/Source/Developer/TraceTools/Private/Widgets/STraceDataFilterWidget.h b/Engine/Source/Developer/TraceTools/Private/Widgets/STraceDataFilterWidget.h
--- a/Engine/Source/Developer/TraceTools/Private/Widgets/STraceDataFilterWidget.h
+++ b/Engine/Source/Developer/TraceTools/Private/Widgets/STraceDataFilterWidget.h
@@ -3,12 +3,12 @@
 #pragma once
 
 #include "Widgets/SCompoundWidget.h"
+#include "Widgets/Views/STileView.h"
 
 class SHorizontalBox;
 class STableViewBase;
-template <typename ItemType> class STreeView;
+template <typename ItemType> class STileView;
 template <typename ItemType> class TTextFilter;
-template <typename ItemType> class TreeFilterHandler;
 
 class ITableRow;
 class ITraceController;
@@ -44,79 +44,65 @@ public:
 	virtual void Tick(const FGeometry& AllottedGeometry, const double InCurrentTime, const float InDeltaTime) override;
 
 protected:
-	/** Callback from SFilterPresetList, should save the current tree view filter state as the specified preset */
+	/** Callback from SFilterPresetList, should save the current tileview filter state as the specified preset */
 	void OnSavePreset(const TSharedPtr<ITraceFilterPreset>& Preset);
 	/** Callback from SFilterPresetList, should update filtering state according to currently active Filter Presets */
 	void OnPresetChanged(const class SFilterPreset& Preset);
-	/** Callback from SFilterPresetList, should highlight the tree view items encompassed by the specified preset */
+	/** Callback from SFilterPresetList, should highlight the tileview items encompassed by the specified preset */
 	void OnHighlightPreset(const TSharedPtr<ITraceFilterPreset>& Preset);
 
 	void OnSearchboxTextChanged(const FText& FilterText);	
 	void OnItemDoubleClicked(TSharedPtr<ITraceObject> InObject) const;
 
 	/** Constructing various widgets */
-	void ConstructTreeview();
+	void ConstructTileView();
 	void ConstructSearchBoxFilter();
-	void ConstructFilterHandler();
 
 	bool HasValidFilterSession() const;
 		
-	void RefreshTreeviewData();
+	void RefreshTileViewData();
 		
-	/** Sets the treeview item expansion state recursively for the specified object */
-	void SetExpansionRecursively(const TSharedPtr<ITraceObject>& InObject, bool bShouldExpandItem) const;
-	void SetParentExpansionRecursively(const TSharedPtr<ITraceObject>& InObject, bool bShouldExpandItem) const;	
-
-	/** Enumerates the (currently selected) treeview items and calls InFunction for each entry */
+	/** Enumerates the (currently selected) tileview items and calls InFunction for each entry */
 	void EnumerateSelectedItems(TFunction<void(TSharedPtr<ITraceObject> InItem)> InFunction) const;
 	void EnumerateAllItems(TFunction<void(TSharedPtr<ITraceObject> InItem)> InFunction) const;
 	bool EnumerateAllItems(TFunction<bool(TSharedPtr<ITraceObject> InItem)> InFunction) const;
 	bool EnumerateSelectedItems(TFunction<bool(TSharedPtr<ITraceObject> InItem)> InFunction) const;
 
-	/** Creates and adds a new treeview object entry, given the specified data */
+	/** Creates and adds a new tileview object entry, given the specified data */
 	TSharedRef<ITraceObject> AddFilterableObject(const FTraceObjectInfo& Event, FString ParentName);
 	
-	/** Treeview callbacks */
+	/** TileView callbacks */
 	TSharedRef<ITableRow> OnGenerateRow(TSharedPtr<ITraceObject> InItem, const TSharedRef<STableViewBase>& OwnerTable);
 	TSharedPtr<SWidget> OnContextMenuOpening() const;
 
 	/** Presets dropdown button callback */
 	TSharedRef<SWidget> MakeAddFilterMenu();
 
-	/** (Re) storing item expansion / selection state */
-	void SaveItemsExpansion();
+	/** (Re) storing item selection state */
 	void SaveItemSelection();
-	void RestoreItemsExpansion();	
 	void RestoreItemSelection();
 
 protected:
-	/** Root level items used by treeview */
-	TArray<TSharedPtr<ITraceObject>> RootItems;
-	/** Dynamically generated array according to treeview filtering */
-	TArray<TSharedPtr<ITraceObject>> TreeItems;
-	/** Flat list of contained treeview items */
-	TArray<TSharedPtr<ITraceObject>> FlatItems;
-	/** Parent child relationships */
-	TMap<TSharedPtr<ITraceObject>, TArray<TSharedPtr<ITraceObject>>> ParentToChild;
-	TMap<TSharedPtr<ITraceObject>, TSharedPtr<ITraceObject>> ChildToParent;
+	/** Flat list of contained items */
+	TArray<TSharedPtr<ITraceObject>> ListItems;
+	/** Dynamically generated array by filtering */
+	TArray<TSharedPtr<ITraceObject>> FilteredListItems;
 
 	/** Analysis filter service for the current analysis session this window is currently representing */
 	TSharedPtr<ISessionTraceFilterService> SessionFilterService;
-	uint32 PreviousSessionHandle;
 	/** Timestamp used for refreshing cached filter data */
 	FDateTime SyncTimeStamp = 0;
 	
 	/** Wrapper for presets drop down button */
 	TSharedPtr<SHorizontalBox> OptionsWidget;
 	
-	/** Treeview widget, represents the currently connected analysis session its filtering state */
-	TSharedPtr<STreeView<TSharedPtr<ITraceObject>>> Treeview;
-	TSharedPtr<TreeFilterHandler<TSharedPtr<ITraceObject>>> TreeviewFilterHandler;
+	/** TileView widget, represents the currently connected analysis session its filtering state */
+	TSharedPtr<STileView<TSharedPtr<ITraceObject>>> TileView;
 
-	/** Flag indicating the treeview should be refreshed */
-	bool bNeedsTreeRefresh;
+	/** Flag indicating the tileview should be refreshed */
+	bool bNeedsListRefresh;
 
-	/** Search box for filtering the Treeview items */
+	/** Search box for filtering the tileview items */
 	TSharedPtr<SSearchBox> SearchBoxWidget;
 	/** The text filter used by the search box */
 	TSharedPtr<TTextFilter<TSharedPtr<ITraceObject>>> SearchBoxWidgetFilter;
@@ -124,22 +110,17 @@ protected:
 	/** Filter presets bar widget */
 	TSharedPtr<SFilterPresetList> FilterPresetsListWidget;
 	
-	/** Widget used for selecting specific AnalysisSession the treeview should represent the filter state for */
-	TSharedPtr<SSessionSelector> SessionSelector;
-
 	/** Whether or not we a currently highlighting a preset */
 	bool bHighlightingPreset;
 
-	/** Cached state of which named entries were expanded / selected */
-	TSet<FString> ExpandedObjectNames;
+	/** Cached state of which named entries were selected */
 	TSet<FString> SelectedObjectNames;
 
-	/** Inter-window cached state of last expansion state */
-	static TSet<FString> LastExpandedObjectNames;
-
 	TSharedPtr<SScrollBar> ExternalScrollbar;
 
 	TSharedPtr<ITraceController> TraceController;
+
+	double AccumulatedTime = 0.0f;
 };
 
 } // namespace UE::TraceTools
\ No newline at end of file
diff --git a/Engine/Source/Developer/TraceTools/Private/Widgets/STraceObjectRowWidget.cpp b/Engine/Source/Developer/TraceTools/Private/Widgets/STraceObjectRowWidget.cpp
--- a/Engine/Source/Developer/TraceTools/Private/Widgets/STraceObjectRowWidget.cpp
+++ b/Engine/Source/Developer/TraceTools/Private/Widgets/STraceObjectRowWidget.cpp
@@ -192,16 +192,27 @@ void STraceObjectRowWidget::Construct(const FArguments& InArgs, const TSharedRef
 		.Padding(0.0f)
 		.Content()
 		[
+			// Enabled checkbox
 			SNew(SHorizontalBox)
 			.IsEnabled(!InObject->IsReadOnly())
+
+			+ SHorizontalBox::Slot()
+			.AutoWidth()
+			.VAlign(VAlign_Center)
+			.Padding(0.0f, 0.0f, 2.0f, 0.0f)
+			[
+				SNew(SFilterCheckboxWidget, InObject)
+			]
+
 			// Name 
 			+ SHorizontalBox::Slot()
 			.FillWidth(1.0f)
-			.Padding(2, 0, 0, 0)
+			.Padding(2.0f, 0.0f, 0.0f, 0.0f)
 			.VAlign(VAlign_Center)
 			[
 				SNew(STextBlock)
 				.Text(Object->GetDisplayText())
+				.ToolTipText(Object->GetTooltipText())
 				.HighlightText_Lambda([this]() -> FText
 				{
 					// Get attribute value, populated by search box for treeview filter-ing
@@ -221,15 +232,6 @@ void STraceObjectRowWidget::Construct(const FArguments& InArgs, const TSharedRef
 					return Value;
 				})
 			]
-
-			// Visibility icon
-			+ SHorizontalBox::Slot()
-			.AutoWidth()
-			.VAlign(VAlign_Center)
-			.Padding(0, 0, 2.f, 0)
-			[
-				SNew(SFilterCheckboxWidget, InObject)
-			]
 		], InOwnerTableView);
 }
 
diff --git a/Engine/Source/Developer/TraceTools/Private/Widgets/STraceStatistics.cpp b/Engine/Source/Developer/TraceTools/Private/Widgets/STraceStatistics.cpp
new file mode 100644
--- /dev/null
+++ b/Engine/Source/Developer/TraceTools/Private/Widgets/STraceStatistics.cpp
@@ -0,0 +1,330 @@
+// Copyright Epic Games, Inc. All Rights Reserved.
+
+#include "STraceStatistics.h"
+
+#include "Internationalization/Text.h"
+#include "SlateOptMacros.h"
+#include "Styling/StyleColors.h"
+#include "Widgets/Layout/SBorder.h"
+#include "Widgets/SBoxPanel.h"
+#include "Widgets/Text/STextBlock.h"
+
+//TraceTools
+#include "Services/SessionTraceControllerFilterService.h"
+#include "TraceToolsStyle.h"
+
+#define LOCTEXT_NAMESPACE "STraceStatistics"
+
+namespace UE::TraceTools
+{
+
+STraceStatistics::STraceStatistics()
+{
+}
+
+STraceStatistics::~STraceStatistics()
+{
+}
+
+BEGIN_SLATE_FUNCTION_BUILD_OPTIMIZATION
+void STraceStatistics::Construct(const FArguments& InArgs, TSharedPtr<ISessionTraceFilterService> InSessionFilterService)
+{
+	SessionFilterService = InSessionFilterService;
+
+	ChildSlot
+	[
+		SNew(SHorizontalBox)
+
+		+ SHorizontalBox::Slot()
+		.AutoWidth()
+		.VAlign(EVerticalAlignment::VAlign_Top)
+		[
+			SNew(SBorder)
+			.BorderImage(FTraceToolsStyle::GetBrush("FilterPresets.BackgroundBorder"))
+			[
+				SNew(SHorizontalBox)
+
+				+ SHorizontalBox::Slot()
+				[
+					SNew(SVerticalBox)
+
+					// Trace Settings
+					+ SVerticalBox::Slot()
+					.HAlign(EHorizontalAlignment::HAlign_Left)
+					.Padding(0.0f, 10.0f, 0.0f, 0.0f)
+					.AutoHeight()
+					[
+						SNew(STextBlock)
+						.Text(LOCTEXT("Trace Settings", "Trace Settings"))
+						.Font(FAppStyle::Get().GetFontStyle("NormalFontBold"))
+					]
+
+					+ SVerticalBox::Slot()
+					.HAlign(EHorizontalAlignment::HAlign_Left)
+					.AutoHeight()
+					[
+						SNew(SHorizontalBox)
+							
+						+ SHorizontalBox::Slot()
+						.Padding(0.0f, 2.0f, 0.0f, 0.0f)
+						.AutoWidth()
+						[
+							SNew(STextBlock)
+							.ColorAndOpacity(FSlateColor(EStyleColor::Foreground))
+							.Text(LOCTEXT("ImportantCache", "Important Events Cache:"))
+						]
+
+						+ SHorizontalBox::Slot()
+						.Padding(2.0f, 2.0f, 0.0f, 0.0f)
+						.AutoWidth()
+						[
+							SNew(STextBlock)
+							.ColorAndOpacity(FSlateColor(EStyleColor::AccentGray))
+							.Text_Lambda([this]() { return this->GetSettingsOnOffText(SessionFilterService->GetSettings().bUseImportantCache); })
+						]
+					]
+
+					+ SVerticalBox::Slot()
+					.HAlign(EHorizontalAlignment::HAlign_Left)
+					.AutoHeight()
+					[
+						SNew(SHorizontalBox)
+							
+						+ SHorizontalBox::Slot()
+						.Padding(0.0f, 2.0f, 0.0f, 0.0f)
+						.AutoWidth()
+						[
+							SNew(STextBlock)
+							.ColorAndOpacity(FSlateColor(EStyleColor::Foreground))
+							.Text(LOCTEXT("WorkerThread", "Worker Thread:"))
+						]
+
+						+ SHorizontalBox::Slot()
+						.Padding(2.0f, 2.0f, 0.0f, 0.0f)
+						.AutoWidth()
+						[
+							SNew(STextBlock)
+							.ColorAndOpacity(FSlateColor(EStyleColor::AccentGray))
+							.Text_Lambda([this]() { return this->GetSettingsOnOffText(SessionFilterService->GetSettings().bUseWorkerThread); })
+						]
+					]
+
+					+ SVerticalBox::Slot()
+					.HAlign(EHorizontalAlignment::HAlign_Left)
+					.AutoHeight()
+					[
+						SNew(SHorizontalBox)
+							
+						+ SHorizontalBox::Slot()
+						.Padding(0.0f, 2.0f, 0.0f, 0.0f)
+						.AutoWidth()
+						[
+							SNew(STextBlock)
+							.ColorAndOpacity(FSlateColor(EStyleColor::Foreground))
+							.Text(LOCTEXT("Tail Size", "Tail Size:"))
+						]
+
+						+ SHorizontalBox::Slot()
+						.Padding(2.0f, 2.0f, 0.0f, 0.0f)
+						.AutoWidth()
+						[
+							SNew(STextBlock)
+							.ColorAndOpacity(FSlateColor(EStyleColor::AccentGray))
+							.Text_Lambda([this]() { return this->GetSettingsMemoryValueText(SessionFilterService->GetSettings().TailSizeBytes); })
+						]
+					]
+				]
+
+				+ SHorizontalBox::Slot()
+				.Padding(30.0f, 0.0f, 0.0f, 0.0f)
+				[
+					SNew(SVerticalBox)
+
+					// Trace statistics
+					+ SVerticalBox::Slot()
+					.HAlign(EHorizontalAlignment::HAlign_Left)
+					.Padding(0.0f, 10.0f, 0.0f, 0.0f)
+					.AutoHeight()
+					[
+						SNew(STextBlock)
+						.Text(LOCTEXT("Statistics", "Statistics"))
+						.Font(FAppStyle::Get().GetFontStyle("NormalFontBold"))
+					]
+
+					+ SVerticalBox::Slot()
+					[
+						SNew(SHorizontalBox)
+							
+						+ SHorizontalBox::Slot()
+						.Padding(0.0f, 2.0, 0.0f, 0.0f)
+						.AutoWidth()
+						[
+							SNew(STextBlock)
+							.ColorAndOpacity(FSlateColor(EStyleColor::Foreground))
+							.Text(LOCTEXT("Bytes Sent", "Bytes Sent:"))
+						]
+
+						+ SHorizontalBox::Slot()
+						.Padding(2.0f, 2.0, 0.0f, 0.0f)
+						[
+							SNew(STextBlock)
+							.ColorAndOpacity(FSlateColor(EStyleColor::AccentGray))
+							.Text_Lambda([this]() { return this->GetStatsMemoryValueText(SessionFilterService->GetStats().BytesSent); })
+						]
+					]
+
+					+ SVerticalBox::Slot()
+					[
+						SNew(SHorizontalBox)
+							
+						+ SHorizontalBox::Slot()
+						.Padding(0.0f, 2.0, 0.0f, 0.0f)
+						.AutoWidth()
+						[
+							SNew(STextBlock)
+							.ColorAndOpacity(FSlateColor(EStyleColor::Foreground))
+							.Text(LOCTEXT("Bytes Traced", "Bytes Traced:"))
+						]
+
+						+ SHorizontalBox::Slot()
+						.Padding(2.0f, 2.0, 0.0f, 0.0f)
+						[
+							SNew(STextBlock)
+							.ColorAndOpacity(FSlateColor(EStyleColor::AccentGray))
+							.Text_Lambda([this]() { return this->GetStatsMemoryValueText(SessionFilterService->GetStats().BytesTraced); })
+						]
+					]
+
+								
+					+ SVerticalBox::Slot()
+					[
+						SNew(SHorizontalBox)
+							
+						+ SHorizontalBox::Slot()
+						.Padding(0.0f, 2.0, 0.0f, 0.0f)
+						.AutoWidth()
+						[
+							SNew(STextBlock)
+							.ColorAndOpacity(FSlateColor(EStyleColor::Foreground))
+							.Text(LOCTEXT("Memory Used", "Memory Used:"))
+						]
+
+						+ SHorizontalBox::Slot()
+						.Padding(2.0f, 2.0, 0.0f, 0.0f)
+						[
+							SNew(STextBlock)
+							.ColorAndOpacity(FSlateColor(EStyleColor::AccentGray))
+							.Text_Lambda([this]() { return this->GetStatsMemoryValueText(SessionFilterService->GetStats().MemoryUsed); })
+						]
+					]
+
+					+ SVerticalBox::Slot()
+					[
+						SNew(SHorizontalBox)
+							
+						+ SHorizontalBox::Slot()
+						.Padding(0.0f, 2.0, 0.0f, 0.0f)
+						.AutoWidth()
+						[
+							SNew(STextBlock)
+							.ColorAndOpacity(FSlateColor(EStyleColor::Foreground))
+							.Text(LOCTEXT("Cache Allocated", "Cache Allocated:"))
+						]
+
+						+ SHorizontalBox::Slot()
+						.Padding(2.0f, 2.0, 0.0f, 0.0f)
+						[
+							SNew(STextBlock)
+							.ColorAndOpacity(FSlateColor(EStyleColor::AccentGray))
+							.Text_Lambda([this]() { return this->GetStatsMemoryValueText(SessionFilterService->GetStats().CacheAllocated); })
+						]
+					]
+
+					+ SVerticalBox::Slot()
+					[
+						SNew(SHorizontalBox)
+							
+						+ SHorizontalBox::Slot()
+						.Padding(0.0f, 2.0, 0.0f, 0.0f)
+						.AutoWidth()
+						[
+							SNew(STextBlock)
+							.ColorAndOpacity(FSlateColor(EStyleColor::Foreground))
+							.Text(LOCTEXT("Cache Used", "Cache Used:"))
+						]
+
+						+ SHorizontalBox::Slot()
+						.Padding(2.0f, 2.0, 0.0f, 0.0f)
+						[
+							SNew(STextBlock)
+							.ColorAndOpacity(FSlateColor(EStyleColor::AccentGray))
+							.Text_Lambda([this]() { return this->GetStatsMemoryValueText(SessionFilterService->GetStats().CacheUsed); })
+						]
+					]
+
+					+ SVerticalBox::Slot()
+					[
+						SNew(SHorizontalBox)
+
+						+ SHorizontalBox::Slot()
+						.Padding(0.0f, 2.0, 0.0f, 0.0f)
+						.AutoWidth()
+						[
+							SNew(STextBlock)
+							.ColorAndOpacity(FSlateColor(EStyleColor::Foreground))
+							.Text(LOCTEXT("Cache Waste", "Cache Waste:"))
+						]
+
+						+ SHorizontalBox::Slot()
+						.Padding(2.0f, 2.0, 0.0f, 0.0f)
+						[
+							SNew(STextBlock)
+							.ColorAndOpacity(FSlateColor(EStyleColor::AccentGray))
+							.Text_Lambda([this]() { return this->GetStatsMemoryValueText(SessionFilterService->GetStats().CacheWaste); })
+						]
+					]
+				]
+			]
+		]
+	];
+}
+END_SLATE_FUNCTION_BUILD_OPTIMIZATION
+
+FText STraceStatistics::GetSettingsOnOffText(bool InValue) const
+{
+	if (!SessionFilterService->HasSettings())
+	{
+		return LOCTEXT("N/A", "N/A");
+	}
+
+	if (InValue)
+	{
+		return LOCTEXT("On", "On");
+	}
+
+	return LOCTEXT("Off", "Off");
+};
+
+FText STraceStatistics::GetSettingsMemoryValueText(uint64 InValue) const
+{
+	if (!SessionFilterService->HasSettings())
+	{
+		return LOCTEXT("N/A", "N/A");
+	}
+
+	return FText::AsMemory(InValue);
+}
+
+FText STraceStatistics::GetStatsMemoryValueText(uint64 InValue) const
+{
+	if (!SessionFilterService->HasStats())
+	{
+		return LOCTEXT("N/A", "N/A");
+	}
+
+	return FText::AsMemory(InValue);
+}
+
+} // namespace UE::TraceTools
+
+#undef LOCTEXT_NAMESPACE
\ No newline at end of file
diff --git a/Engine/Source/Developer/TraceTools/Private/Widgets/STraceStatistics.h b/Engine/Source/Developer/TraceTools/Private/Widgets/STraceStatistics.h
new file mode 100644
--- /dev/null
+++ b/Engine/Source/Developer/TraceTools/Private/Widgets/STraceStatistics.h
@@ -0,0 +1,41 @@
+// Copyright Epic Games, Inc. All Rights Reserved.
+
+#pragma once
+
+#include "ITraceController.h"
+#include "Widgets/DeclarativeSyntaxSupport.h"
+#include "Widgets/SCompoundWidget.h"
+
+class ITraceController;
+
+namespace UE::TraceTools
+{
+
+class ISessionTraceFilterService;
+
+/**
+ * A widget that displays trace settings and trace statistics
+ */
+class STraceStatistics
+	: public SCompoundWidget
+{
+public:
+	SLATE_BEGIN_ARGS(STraceStatistics) { }
+	SLATE_END_ARGS()
+
+public:
+	STraceStatistics();
+	virtual ~STraceStatistics();
+
+	void Construct( const FArguments& InArgs, TSharedPtr<ISessionTraceFilterService> SessionFilterService);
+
+private:
+	FText GetSettingsOnOffText(bool InValue) const;
+	FText GetSettingsMemoryValueText(uint64 InValue) const;
+	FText GetStatsMemoryValueText(uint64 InValue) const;
+
+private:
+	TSharedPtr<ISessionTraceFilterService> SessionFilterService;
+};
+
+} // namespace UE::TraceTools
\ No newline at end of file
diff --git a/Engine/Source/Runtime/SessionServices/Private/TraceController.cpp b/Engine/Source/Runtime/SessionServices/Private/TraceController.cpp
--- a/Engine/Source/Runtime/SessionServices/Private/TraceController.cpp
+++ b/Engine/Source/Runtime/SessionServices/Private/TraceController.cpp
@@ -105,16 +105,17 @@ void FTraceController::OnDiscoveryResponse(const FTraceControlDiscovery& Message
 	}
 	
 	UpdateStatus(Message, Instance->Status);
-	
+
 	// This is the application session id, which is not necessarily the
 	// same as trace session (maybe overridden by commandline)
 	Instance->Status.SessionId = Message.SessionId;
 	Instance->Status.InstanceId = Message.InstanceId;
 
-	StatusReceivedEvent.Broadcast(Instance->Status, FTraceStatus::EUpdateType::All, Instance->Commands);
+	FTraceStatus::EUpdateType UpdateType = FTraceStatus::EUpdateType::ChannelsDesc | FTraceStatus::EUpdateType::ChannelsStatus | FTraceStatus::EUpdateType::Status;
+	StatusReceivedEvent.Broadcast(Instance->Status, UpdateType, Instance->Commands);
 	if (SelectedInstanceIds.Contains(Message.InstanceId))
 	{
-		SelectedSessionStatusReceivedEvent.Broadcast(Instance->Status, FTraceStatus::EUpdateType::All, Instance->Commands);
+		SelectedSessionStatusReceivedEvent.Broadcast(Instance->Status, UpdateType, Instance->Commands);
 	}
 }
 
@@ -246,7 +247,8 @@ void FTraceController::OnInstanceSelectionChanged(const TSharedPtr<ISessionInsta
 		if (const auto Instance = Instances.Find(*Address))
 		{
 			FTraceStatus& Status = Instance->Status;
-			SelectedSessionStatusReceivedEvent.Broadcast(Status, FTraceStatus::EUpdateType::All, Instance->Commands);
+			FTraceStatus::EUpdateType UpdateType = FTraceStatus::EUpdateType::ChannelsDesc | FTraceStatus::EUpdateType::ChannelsStatus | FTraceStatus::EUpdateType::Status;
+			SelectedSessionStatusReceivedEvent.Broadcast(Status, UpdateType, Instance->Commands);
 		}
 	}
 	else
