Fix multiple session support in Trace Control Window.
Active presets now reflect the state of the channels.

diff --git a/Engine/Source/Developer/TraceTools/Private/Models/TraceChannel.cpp b/Engine/Source/Developer/TraceTools/Private/Models/TraceChannel.cpp
--- a/Engine/Source/Developer/TraceTools/Private/Models/TraceChannel.cpp
+++ b/Engine/Source/Developer/TraceTools/Private/Models/TraceChannel.cpp
@@ -13,14 +13,14 @@ namespace UE::TraceTools
 FTraceChannel::FTraceChannel(FString InName,
 							 FString InDescription,
 							 FString InParentName, 
-							 uint32 InHash,
+							 uint32 InId,
 							 bool bInEnabled,
 							 bool bInReadOnly,
 							 TSharedPtr<ISessionTraceFilterService> InFilterService) 
 	: Name(InName),
 	  Description(InDescription),
 	  ParentName(InParentName),
-	  Hash(InHash),
+	  Id(InId),
 	  bFiltered(!bInEnabled), 
 	  bIsPending(false),
 	  bReadOnly(bInReadOnly),
diff --git a/Engine/Source/Developer/TraceTools/Private/Models/TraceChannel.h b/Engine/Source/Developer/TraceTools/Private/Models/TraceChannel.h
--- a/Engine/Source/Developer/TraceTools/Private/Models/TraceChannel.h
+++ b/Engine/Source/Developer/TraceTools/Private/Models/TraceChannel.h
@@ -12,7 +12,7 @@ class ISessionTraceFilterService;
 class FTraceChannel : public ITraceObject
 {
 public:
-	FTraceChannel(FString InName, FString InDescription, FString InParentName, uint32 InHash, bool bInEnabled, bool bInReadOnly, TSharedPtr<ISessionTraceFilterService> InFilterService);
+	FTraceChannel(FString InName, FString InDescription, FString InParentName, uint32 InId, bool bInEnabled, bool bInReadOnly, TSharedPtr<ISessionTraceFilterService> InFilterService);
 
 	/** Begin ITraceObject overrides */
 	virtual FText GetDisplayText() const override;
@@ -34,8 +34,8 @@ protected:
 	FString Description;
 	/** Channel's parent (group) name */
 	FString ParentName;
-	/** Hash representing this channel uniquely */
-	uint32 Hash;
+	/** Channel's id */
+	uint32 Id;
 	
 	/** Whether or not this channel is filtered out, true = filtered; false = not filtered */
 	bool bFiltered;
diff --git a/Engine/Source/Developer/TraceTools/Private/Services/ISessionTraceFilterService.h b/Engine/Source/Developer/TraceTools/Private/Services/ISessionTraceFilterService.h
--- a/Engine/Source/Developer/TraceTools/Private/Services/ISessionTraceFilterService.h
+++ b/Engine/Source/Developer/TraceTools/Private/Services/ISessionTraceFilterService.h
@@ -20,8 +20,7 @@ struct FTraceObjectInfo
 	FString Description;
 	bool bEnabled;
 	bool bReadOnly;
-	uint32 Hash;
-	uint32 OwnerHash;
+	uint32 Id;
 
 	bool operator<(const FTraceObjectInfo& Other) const
 	{
@@ -37,6 +36,9 @@ public:
 
 	/** Returns the root level set of objects */
 	virtual void GetRootObjects(TArray<FTraceObjectInfo>& OutObjects) const = 0;
+	
+	/** Returns the object with the specified name */
+	virtual const FTraceObjectInfo* GetObject(const FString& Name ) const = 0;
 
 	/** Set the filtered state for an individual object by its hash */
 	virtual void SetObjectFilterState(const FString& InObjectName, const bool bFilterState) = 0;
diff --git a/Engine/Source/Developer/TraceTools/Private/Services/SessionTraceControllerFilterService.cpp b/Engine/Source/Developer/TraceTools/Private/Services/SessionTraceControllerFilterService.cpp
--- a/Engine/Source/Developer/TraceTools/Private/Services/SessionTraceControllerFilterService.cpp
+++ b/Engine/Source/Developer/TraceTools/Private/Services/SessionTraceControllerFilterService.cpp
@@ -2,18 +2,32 @@
 
 #include "SessionTraceControllerFilterService.h"
 
-#include "Models/ITraceFilterPreset.h"
+#include "Hash/xxhash.h"
 #include "Misc/CoreDelegates.h"
 
+// TraceTools
+#include "Models/ITraceFilterPreset.h"
+
 namespace UE::TraceTools
 {
 
+uint64 HashName(FStringView Name)
+{
+	// Strip plurals and convert to upper case.
+	TStringBuilder<48> TransformedBuffer;
+	TransformedBuffer << (Name.EndsWith('s') ? Name.LeftChop(1) : Name);
+	FCString::Strupr(TransformedBuffer.GetData(), TransformedBuffer.Len());
+
+	return FXxHash64::HashBuffer(TransformedBuffer.GetData(), TransformedBuffer.Len() * sizeof(TCHAR)).Hash;
+}
+
 FSessionTraceControllerFilterService::FSessionTraceControllerFilterService(TSharedPtr<ITraceController> InTraceController)
 {
 	FCoreDelegates::OnEndFrame.AddRaw(this, &FSessionTraceControllerFilterService::OnApplyChannelChanges);
 
 	TraceController = InTraceController;
 	TraceController->OnSelectedSessionStatusReceived().AddRaw(this, &FSessionTraceControllerFilterService::OnTraceStatusUpdated);
+	TraceController->OnSessionSelectionChanged().AddRaw(this, &FSessionTraceControllerFilterService::OnSessionSelectionChanged);
 }
 
 FSessionTraceControllerFilterService::~FSessionTraceControllerFilterService()
@@ -21,11 +35,20 @@ FSessionTraceControllerFilterService::~FSessionTraceControllerFilterService()
 	FCoreDelegates::OnEndFrame.RemoveAll(this);
 
 	TraceController->OnSelectedSessionStatusReceived().RemoveAll(this);
+	TraceController->OnSessionSelectionChanged().RemoveAll(this);
 }
 
 void FSessionTraceControllerFilterService::GetRootObjects(TArray<FTraceObjectInfo>& OutObjects) const
 {
-	OutObjects.Append(Objects);
+	for (auto& Pair : Objects)
+	{
+		OutObjects.Add(Pair.Value);
+	}
+}
+
+const FTraceObjectInfo* FSessionTraceControllerFilterService::GetObject(const FString& Name) const
+{
+	return Objects.Find(HashName(Name));
 }
 
 const FDateTime& FSessionTraceControllerFilterService::GetTimestamp() const
@@ -73,9 +96,9 @@ void FSessionTraceControllerFilterService::UpdateFilterPreset(const TSharedPtr<I
 
 void FSessionTraceControllerFilterService::DisableAllChannels()
 {
-	for (FTraceObjectInfo& ObjectInfo : Objects)
+	for (auto& ObjectInfo : Objects)
 	{
-		ObjectInfo.bEnabled = false;
+		ObjectInfo.Value.bEnabled = false;
 	}
 }
 
@@ -120,13 +143,12 @@ void FSessionTraceControllerFilterService::UpdateChannels(const FTraceStatus& In
 
 	for (auto& Entry : Channels)
 	{
-		FTraceObjectInfo& EventInfo = Objects.AddDefaulted_GetRef();
+		FTraceObjectInfo& EventInfo = Objects.Add(HashName(Entry.Value.Name));
 		EventInfo.Name = Entry.Value.Name;
 		EventInfo.Description = Entry.Value.Description;
 		EventInfo.bEnabled = Entry.Value.bEnabled;
 		EventInfo.bReadOnly = Entry.Value.bReadOnly;
-		EventInfo.Hash = Entry.Value.Id;
-		EventInfo.OwnerHash = 0;
+		EventInfo.Id = Entry.Value.Id;
 	}
 }
 
@@ -169,4 +191,10 @@ const FTraceStatus::FStats& FSessionTraceControllerFilterService::GetStats() con
 	return Stats;
 }
 
+void FSessionTraceControllerFilterService::OnSessionSelectionChanged()
+{
+	bHasStats = false;
+	bHasSettings = false;
+}
+
 } // namespace UE::TraceTools
\ No newline at end of file
diff --git a/Engine/Source/Developer/TraceTools/Private/Services/SessionTraceControllerFilterService.h b/Engine/Source/Developer/TraceTools/Private/Services/SessionTraceControllerFilterService.h
--- a/Engine/Source/Developer/TraceTools/Private/Services/SessionTraceControllerFilterService.h
+++ b/Engine/Source/Developer/TraceTools/Private/Services/SessionTraceControllerFilterService.h
@@ -25,6 +25,8 @@ public:
 
 	/** Begin ISessionTraceFilterService overrides */
 	virtual void GetRootObjects(TArray<FTraceObjectInfo>& OutObjects) const override;
+	virtual const FTraceObjectInfo* GetObject(const FString& Name) const override;
+	
 	virtual const FDateTime& GetTimestamp() const override;
 	virtual void SetObjectFilterState(const FString& InObjectName, const bool bFilterState) override;
 	virtual void UpdateFilterPreset(const TSharedPtr<ITraceFilterPreset> InPreset, bool IsEnabled) override;
@@ -43,16 +45,17 @@ protected:
 	/** Retrieves channels names from provider and marks them all as disabled */
 	void DisableAllChannels();
 
-	void RetrieveAndStoreStartupChannels();
-
 	void OnTraceStatusUpdated(const FTraceStatus& InStatus, FTraceStatus::EUpdateType InUpdateType, ITraceControllerCommands& Commands);
 
 	void UpdateChannels(const FTraceStatus& InStatus);
 
+	void OnSessionSelectionChanged();
+
 protected:
 	TSharedPtr<ITraceController> TraceController;
 
-	TArray<FTraceObjectInfo> Objects;
+	/** A map with the key formed by hashing the object name and the object as the value.*/
+	TMap<uint64, FTraceObjectInfo> Objects;
 
 	/** Names of channels that were either enabled or disabled during the duration of this frame */
 	TSet<FString> FrameEnabledChannels;
diff --git a/Engine/Source/Developer/TraceTools/Private/Widgets/SFilterPreset.cpp b/Engine/Source/Developer/TraceTools/Private/Widgets/SFilterPreset.cpp
--- a/Engine/Source/Developer/TraceTools/Private/Widgets/SFilterPreset.cpp
+++ b/Engine/Source/Developer/TraceTools/Private/Widgets/SFilterPreset.cpp
@@ -105,6 +105,11 @@ bool SFilterPreset::IsEnabled() const
 	return bEnabled;
 }
 
+void SFilterPreset::MarkAsEnabled(bool InEnabled)
+{
+	bEnabled = InEnabled;
+}
+
 const TSharedPtr<ITraceFilterPreset>& SFilterPreset::GetFilterPreset() const
 {
 	return FilterPreset;
diff --git a/Engine/Source/Developer/TraceTools/Private/Widgets/SFilterPreset.h b/Engine/Source/Developer/TraceTools/Private/Widgets/SFilterPreset.h
--- a/Engine/Source/Developer/TraceTools/Private/Widgets/SFilterPreset.h
+++ b/Engine/Source/Developer/TraceTools/Private/Widgets/SFilterPreset.h
@@ -137,6 +137,9 @@ public:
 	virtual FReply OnMouseMove(const FGeometry& MyGeometry, const FPointerEvent& MouseEvent) override;
 	/** End SCompoundWidget overrides */
 
+	/** Sets the enabled state but doesn't call the callback that changes the channel's state */
+	void MarkAsEnabled(bool InEnabled);
+
 	/** Sets whether or not this preset is applied to the combined filtering state */
 	void SetEnabled(bool InEnabled);
 
diff --git a/Engine/Source/Developer/TraceTools/Private/Widgets/SFilterPresetList.cpp b/Engine/Source/Developer/TraceTools/Private/Widgets/SFilterPresetList.cpp
--- a/Engine/Source/Developer/TraceTools/Private/Widgets/SFilterPresetList.cpp
+++ b/Engine/Source/Developer/TraceTools/Private/Widgets/SFilterPresetList.cpp
@@ -12,6 +12,7 @@
 // TraceTools
 #include "Models/TraceFilterPresets.h"
 #include "TraceToolsModule.h"
+#include "Services/SessionTraceControllerFilterService.h"
 
 #if WITH_EDITOR
 #include "SSettingsEditorCheckoutNotice.h"
@@ -22,12 +23,14 @@
 namespace UE::TraceTools
 {
 
-void SFilterPresetList::Construct( const FArguments& InArgs )
+void SFilterPresetList::Construct( const FArguments& InArgs, TSharedPtr<ISessionTraceFilterService> InSessionFilterService)
 {
 	OnPresetChanged = InArgs._OnPresetChanged;
 	OnSavePreset = InArgs._OnSavePreset;
 	OnHighlightPreset = InArgs._OnHighlightPreset;
 
+	SessionFilterService = InSessionFilterService;
+
 	FilterBox = SNew(SWrapBox)
 				.UseAllottedSize(true);
 
@@ -50,6 +53,29 @@ SFilterPresetList::~SFilterPresetList()
 	SaveSettings(FTraceToolsModule::TraceFiltersIni);
 }
 
+void SFilterPresetList::RefreshPresetEnabledState()
+{
+	TArray<FString> Names;
+
+	for (TSharedRef<SFilterPreset> PresetWidget : Presets)
+	{
+		Names.Empty();
+		bool bPresetIsEnabled = true;
+		PresetWidget->GetFilterPreset()->GetAllowlistedNames(Names);
+		for (const FString& Name : Names)
+		{
+			const FTraceObjectInfo* Object = SessionFilterService->GetObject(Name);
+			if (Object == nullptr || !Object->bEnabled)
+			{
+				bPresetIsEnabled = false;
+				break;
+			}
+		}
+
+		PresetWidget->MarkAsEnabled(bPresetIsEnabled);
+	}
+}
+
 void SFilterPresetList::RefreshFilterPresets()
 {
 	FSlateApplication::Get().DismissAllMenus();
@@ -60,11 +86,6 @@ void SFilterPresetList::RefreshFilterPresets()
 		{
 			const TSharedRef<SFilterPreset>& FilterToRemove = *FilterIt;
 			ActiveFilterNames.Add(FilterToRemove->GetFilterPreset()->GetName());
-
-			if (FilterToRemove->IsEnabled())
-			{
-				EnabledFilterNames.Add(FilterToRemove->GetFilterPreset()->GetName());
-			}
 		}
 
 		FilterBox->ClearChildren();
@@ -91,15 +112,12 @@ void SFilterPresetList::RefreshFilterPresets()
 		if (ActiveFilterNames.Contains(Preset->GetName()))
 		{
 			TSharedRef<SFilterPreset> Filter = AddFilterPreset(Preset);
-			if (EnabledFilterNames.Contains(Preset->GetName()))
-			{
-				Filter->SetEnabled(true);
-			}
 		}
 	}
 
+	RefreshPresetEnabledState();
+
 	ActiveFilterNames.Empty();
-	EnabledFilterNames.Empty();
 }
 
 FReply SFilterPresetList::OnMouseButtonUp( const FGeometry& MyGeometry, const FPointerEvent& MouseEvent )
@@ -133,25 +151,16 @@ void SFilterPresetList::SaveSettings(const FString& IniFilename) const
 	FStringView Separator = TEXTVIEW(";");
 	
 	FString ActivePresetsString;
-	FString EnabledPresetsString;
 
 	for (const TSharedRef<SFilterPreset>& Preset : Presets)
 	{
 		ActivePresetsString.Append(Preset->GetFilterPreset()->GetName());
 		ActivePresetsString.Append(Separator);
-
-		if (Preset->IsEnabled())
-		{
-			EnabledPresetsString.Append(Preset->GetFilterPreset()->GetName());
-			EnabledPresetsString.Append(Separator);
-		}
 	}
 
 	ActivePresetsString.RemoveFromEnd(Separator);
-	EnabledPresetsString.RemoveFromEnd(Separator);
 
 	GConfig->SetString(IniSectionName.GetData(), IniActivePresetsKey.GetData(), *ActivePresetsString, IniFilename);
-	GConfig->SetString(IniSectionName.GetData(), IniEnabledPresetsKey.GetData(), *EnabledPresetsString, IniFilename);
 
 	GConfig->Flush(false, IniFilename);
 }
@@ -169,7 +178,6 @@ void SFilterPresetList::LoadSettings(const FString& IniFilename)
 
 	FStringView Separator = TEXT(";");
 	ActivePresetsString.ParseIntoArray(ActiveFilterNames, Separator.GetData());
-	EnabledPresetsString.ParseIntoArray(EnabledFilterNames, Separator.GetData());
 }
 
 bool SFilterPresetList::HasAnyPresets() const
diff --git a/Engine/Source/Developer/TraceTools/Private/Widgets/SFilterPresetList.h b/Engine/Source/Developer/TraceTools/Private/Widgets/SFilterPresetList.h
--- a/Engine/Source/Developer/TraceTools/Private/Widgets/SFilterPresetList.h
+++ b/Engine/Source/Developer/TraceTools/Private/Widgets/SFilterPresetList.h
@@ -13,6 +13,7 @@ class USharedTraceFilterPresetContainer;
 namespace UE::TraceTools
 {
 
+class ISessionTraceFilterService;
 class SFilterPreset;
 
 /**
@@ -34,7 +35,7 @@ public:
 	SLATE_END_ARGS()
 
 	/** Constructs this widget with InArgs */		
-	void Construct( const FArguments& InArgs );
+	void Construct( const FArguments& InArgs, TSharedPtr<ISessionTraceFilterService> InSessionFilterService);
 
 	virtual ~SFilterPresetList();
 
@@ -48,6 +49,9 @@ public:
 	/** Returns the context menu widget generated by MakeFilterPresetsMenu */
 	TSharedRef<SWidget> ExternalMakeFilterPresetsMenu();	
 
+	/** Sets channels as enabled/disabled based on the current object list */
+	void RefreshPresetEnabledState();
+
 private:
 	/** Refreshes the filter preset objects and child widgets */
 	void RefreshFilterPresets();
@@ -128,7 +132,9 @@ private:
 
 	/** Cached arrays of currently active and enabled filter names, used for saving/restoring state during refreshes */
 	TArray<FString> ActiveFilterNames;
-	TArray<FString> EnabledFilterNames;
+
+	/** A pointer to the SessionFilterService to autodetect enabled presets. */
+	TSharedPtr<ISessionTraceFilterService> SessionFilterService;
 };
 
 } // namespace UE::TraceTools
\ No newline at end of file
diff --git a/Engine/Source/Developer/TraceTools/Private/Widgets/STraceControlToolbar.cpp b/Engine/Source/Developer/TraceTools/Private/Widgets/STraceControlToolbar.cpp
--- a/Engine/Source/Developer/TraceTools/Private/Widgets/STraceControlToolbar.cpp
+++ b/Engine/Source/Developer/TraceTools/Private/Widgets/STraceControlToolbar.cpp
@@ -35,13 +35,14 @@ void STraceControlToolbar::Construct(const FArguments& InArgs, const TSharedRef<
 	FTraceControlCommands::Register();
 	TraceController = InTraceController;
 
-	OnStatusReceivedDelegate = TraceController->OnStatusReceived().AddSP(this, &STraceControlToolbar::OnTraceStatusUpdated);
+	OnStatusReceivedDelegate = TraceController->OnSelectedSessionStatusReceived().AddSP(this, &STraceControlToolbar::OnTraceStatusUpdated);
 
 	BindCommands(CommandList);
 
 	// create the toolbar
-	FToolBarBuilder Toolbar(CommandList, FMultiBoxCustomization::None);
+	FSlimHorizontalToolBarBuilder Toolbar(CommandList, FMultiBoxCustomization::None);
 	{
+		Toolbar.SetStyle(&FTraceToolsStyle::Get(), "TraceControlToolbar");
 		Toolbar.AddComboButton(
 			FUIAction(),
 			FOnGetContent::CreateSP(this, &STraceControlToolbar::BuildTraceTargetMenu, CommandList),
@@ -52,6 +53,8 @@ void STraceControlToolbar::Construct(const FArguments& InArgs, const TSharedRef<
 
 		Toolbar.AddSeparator();
 
+		Toolbar.SetLabelVisibility(EVisibility::Collapsed);
+
 		Toolbar.AddToolBarButton(FTraceControlCommands::Get().StartTrace);
 		Toolbar.AddToolBarButton(FTraceControlCommands::Get().StopTrace);
 		Toolbar.AddToolBarButton(FTraceControlCommands::Get().TraceSnapshot);
@@ -309,12 +312,7 @@ FText STraceControlToolbar::GetTraceTargetLabelText() const
 
 FText STraceControlToolbar::GetTraceTargetTooltipText() const
 {
-	if (TraceTarget == ETraceTarget::Server)
-	{
-		return LOCTEXT("TraceTargetServerTooltip", "Set the Unreal Trace Server as the trace target.");
-	}
-
-	return LOCTEXT("TraceTargetFileTooltip", "Set File as the trace target.");
+	return LOCTEXT("TraceTargetTooltip", "Set the trace target. Can only be set when trace is not running.");
 }
 
 FSlateIcon STraceControlToolbar::GetTraceTargetIcon() const
diff --git a/Engine/Source/Developer/TraceTools/Private/Widgets/STraceDataFilterWidget.cpp b/Engine/Source/Developer/TraceTools/Private/Widgets/STraceDataFilterWidget.cpp
--- a/Engine/Source/Developer/TraceTools/Private/Widgets/STraceDataFilterWidget.cpp
+++ b/Engine/Source/Developer/TraceTools/Private/Widgets/STraceDataFilterWidget.cpp
@@ -34,6 +34,7 @@ STraceDataFilterWidget::STraceDataFilterWidget() : bNeedsListRefresh(false), bHi
 
 STraceDataFilterWidget::~STraceDataFilterWidget()
 {
+	TraceController->OnSessionSelectionChanged().RemoveAll(this);
 }
 
 void STraceDataFilterWidget::Construct(const FArguments& InArgs, TSharedPtr<ITraceController> InTraceController)
@@ -41,6 +42,8 @@ void STraceDataFilterWidget::Construct(const FArguments& InArgs, TSharedPtr<ITra
 	TraceController = InTraceController;
 	SessionFilterService = MakeShareable(new FSessionTraceControllerFilterService(InTraceController));
 
+	TraceController->OnSessionSelectionChanged().AddSP(this, &STraceDataFilterWidget::OnSessionSelectionChanged);
+
 	SAssignNew(ExternalScrollbar, SScrollBar)
 	.AlwaysShowScrollbar(true);
 
@@ -114,7 +117,7 @@ void STraceDataFilterWidget::Construct(const FArguments& InArgs, TSharedPtr<ITra
 				+SHorizontalBox::Slot()
 				.FillWidth(1.0f)
 				[
-					SAssignNew(FilterPresetsListWidget, SFilterPresetList)
+					SAssignNew(FilterPresetsListWidget, SFilterPresetList, SessionFilterService)
 					.OnPresetChanged(this, &STraceDataFilterWidget::OnPresetChanged)
 					.OnSavePreset(this, &STraceDataFilterWidget::OnSavePreset)
 					.OnHighlightPreset(this, &STraceDataFilterWidget::OnHighlightPreset)
@@ -130,7 +133,7 @@ void STraceDataFilterWidget::Construct(const FArguments& InArgs, TSharedPtr<ITra
 					SNew(SHorizontalBox)
 					.Visibility_Lambda([this]() -> EVisibility 
 					{
-						return HasValidFilterSession() ? EVisibility::Collapsed : EVisibility::Visible;
+						return ShouldShowBanner() ? EVisibility::Visible : EVisibility::Collapsed;
 					})
 
 					+ SHorizontalBox::Slot()
@@ -148,7 +151,7 @@ void STraceDataFilterWidget::Construct(const FArguments& InArgs, TSharedPtr<ITra
 					.Padding(4.0f, 0.0f, 0.0f, 0.0f)
 					[
 						SNew(STextBlock)
-						.Text(LOCTEXT("NoSessionSelectedWarning", "Please select an instance from the Session Browser."))
+						.Text(LOCTEXT("NoSessionSelectedWarning", "Please select an active instance from the Session Browser."))
 						.ColorAndOpacity(FLinearColor(1.0f, 1.0f, 1.0f, 1.0f))
 					]
 				]
@@ -189,10 +192,6 @@ void STraceDataFilterWidget::Construct(const FArguments& InArgs, TSharedPtr<ITra
 					+ SVerticalBox::Slot()
 					[
 						SNew(STraceStatistics, SessionFilterService)
-						.Visibility_Lambda([this]() -> EVisibility
-						{
-							return HasValidFilterSession() ? EVisibility::Visible : EVisibility::Collapsed;
-						})
 					]
 				]
 			]
@@ -201,7 +200,7 @@ void STraceDataFilterWidget::Construct(const FArguments& InArgs, TSharedPtr<ITra
 
 	/** Setup attribute to enable/disabled all widgets according to whether or not there is a valid session to represent */
 	TAttribute<bool> EnabledAttribute;
-	EnabledAttribute.Bind(this, &STraceDataFilterWidget::HasValidFilterSession);
+	EnabledAttribute.Bind(this, &STraceDataFilterWidget::HasValidData);
 
 	TileView->SetEnabled(EnabledAttribute);
 	OptionsWidget->SetEnabled(EnabledAttribute);
@@ -428,16 +427,21 @@ void STraceDataFilterWidget::RestoreItemSelection()
 
 TSharedRef<ITraceObject> STraceDataFilterWidget::AddFilterableObject(const FTraceObjectInfo& Event, FString ParentName)
 {
-	TSharedRef<FTraceChannel> SharedItem = MakeShareable(new FTraceChannel(Event.Name, Event.Description, ParentName, Event.Hash, Event.bEnabled, Event.bReadOnly, SessionFilterService));
+	TSharedRef<FTraceChannel> SharedItem = MakeShareable(new FTraceChannel(Event.Name, Event.Description, ParentName, Event.Id, Event.bEnabled, Event.bReadOnly, SessionFilterService));
 
 	ListItems.Add(SharedItem);
 
 	return SharedItem;
 }
 
-bool STraceDataFilterWidget::HasValidFilterSession() const
+bool STraceDataFilterWidget::HasValidData() const
+{
+	return bHasChannelData && TraceController->GetNumSelectedInstances() == 1 && TraceController->HasAvailableSelectedInstance() && SessionFilterService.IsValid();
+}
+
+bool STraceDataFilterWidget::ShouldShowBanner() const
 {
-	return TraceController->HasAvailableSelectedInstance() && SessionFilterService.IsValid();
+	return TraceController->GetNumSelectedInstances() != 1 || !TraceController->HasAvailableSelectedInstance() || !SessionFilterService.IsValid();
 }
 
 void STraceDataFilterWidget::EnumerateSelectedItems(TFunction<void(TSharedPtr<ITraceObject> InItem)> InFunction) const
@@ -506,6 +510,8 @@ void STraceDataFilterWidget::RefreshTileViewData()
 		}
 
 		RestoreItemSelection();
+
+		FilterPresetsListWidget->RefreshPresetEnabledState();
 	}
 }
 
@@ -519,6 +525,7 @@ void STraceDataFilterWidget::Tick(const FGeometry& AllottedGeometry, const doubl
 		{
 			RefreshTileViewData();
 			bNeedsListRefresh = true;
+			bHasChannelData = true;
 		}
 
 		if (bNeedsListRefresh)
@@ -555,6 +562,11 @@ void STraceDataFilterWidget::Tick(const FGeometry& AllottedGeometry, const doubl
 	}
 }
 
+void STraceDataFilterWidget::OnSessionSelectionChanged()
+{
+	bHasChannelData = false;
+}
+
 } // namespace UE::TraceTools
 
 #undef LOCTEXT_NAMESPACE // "STraceDataFilterWidget"
diff --git a/Engine/Source/Developer/TraceTools/Private/Widgets/STraceDataFilterWidget.h b/Engine/Source/Developer/TraceTools/Private/Widgets/STraceDataFilterWidget.h
--- a/Engine/Source/Developer/TraceTools/Private/Widgets/STraceDataFilterWidget.h
+++ b/Engine/Source/Developer/TraceTools/Private/Widgets/STraceDataFilterWidget.h
@@ -58,7 +58,8 @@ protected:
 	void ConstructTileView();
 	void ConstructSearchBoxFilter();
 
-	bool HasValidFilterSession() const;
+	bool HasValidData() const;
+	bool ShouldShowBanner() const;
 		
 	void RefreshTileViewData();
 		
@@ -82,6 +83,8 @@ protected:
 	void SaveItemSelection();
 	void RestoreItemSelection();
 
+	void OnSessionSelectionChanged();
+
 protected:
 	/** Flat list of contained items */
 	TArray<TSharedPtr<ITraceObject>> ListItems;
@@ -121,6 +124,8 @@ protected:
 	TSharedPtr<ITraceController> TraceController;
 
 	double AccumulatedTime = 0.0f;
+
+	bool bHasChannelData = false;
 };
 
 } // namespace UE::TraceTools
\ No newline at end of file
diff --git a/Engine/Source/Runtime/SessionServices/Private/SessionManager.cpp b/Engine/Source/Runtime/SessionServices/Private/SessionManager.cpp
--- a/Engine/Source/Runtime/SessionServices/Private/SessionManager.cpp
+++ b/Engine/Source/Runtime/SessionServices/Private/SessionManager.cpp
@@ -135,6 +135,12 @@ bool FSessionManager::SelectSession(const TSharedPtr<ISessionInfo>& Session)
 		return false;
 	}
 
+	// Notify listeners that previously selected sessions have been deselected.
+	for (TSharedPtr<ISessionInstanceInfo> Instance : SelectedInstances)
+	{
+		InstanceSelectionChangedDelegate.Broadcast(Instance, false);
+	}
+
 	// set selection
 	SelectedInstances.Empty();
 	SelectedSession = Session;
diff --git a/Engine/Source/Runtime/SessionServices/Private/TraceController.cpp b/Engine/Source/Runtime/SessionServices/Private/TraceController.cpp
--- a/Engine/Source/Runtime/SessionServices/Private/TraceController.cpp
+++ b/Engine/Source/Runtime/SessionServices/Private/TraceController.cpp
@@ -255,6 +255,8 @@ void FTraceController::OnInstanceSelectionChanged(const TSharedPtr<ISessionInsta
 	{
 		SendDiscoveryPing(InstanceInfo);
 	}
+
+	SessionSelectionChangedEvent.Broadcast();
 }
 
 void FTraceController::UpdateStatus(const FTraceControlStatus& Message, FTraceStatus& Status)
diff --git a/Engine/Source/Runtime/SessionServices/Private/TraceController.h b/Engine/Source/Runtime/SessionServices/Private/TraceController.h
--- a/Engine/Source/Runtime/SessionServices/Private/TraceController.h
+++ b/Engine/Source/Runtime/SessionServices/Private/TraceController.h
@@ -44,6 +44,14 @@ private:
 		return SelectedSessionStatusReceivedEvent;
 	}
 
+	DECLARE_DERIVED_EVENT(FTraceController, ITraceController::FSessionSelectionChanged, FSessionSelectionChanged);
+	virtual FSessionSelectionChanged& OnSessionSelectionChanged() override
+	{
+		return SessionSelectionChangedEvent;
+	}
+
+	virtual uint32 GetNumSelectedInstances() override { return SelectedInstanceIds.Num(); }
+
 	/* Message handlers */
 	void OnNotification(const FMessageBusNotification& MessageBusNotification);
 	void OnDiscoveryResponse(const FTraceControlDiscovery& Message, const TSharedRef<IMessageContext>& Context);
@@ -97,6 +105,9 @@ private:
 	/** Event for status updates on a selected session */
 	FStatusRecievedEvent SelectedSessionStatusReceivedEvent;
 
+	/** Event that triggers when the session selection changes */
+	FSessionSelectionChanged SessionSelectionChangedEvent;
+
 	/** Lock to protect access to Instances list */
 	FRWLock InstancesLock;
 
diff --git a/Engine/Source/Runtime/SessionServices/Public/ITraceController.h b/Engine/Source/Runtime/SessionServices/Public/ITraceController.h
--- a/Engine/Source/Runtime/SessionServices/Public/ITraceController.h
+++ b/Engine/Source/Runtime/SessionServices/Public/ITraceController.h
@@ -159,6 +159,16 @@ public:
 	 */
 	virtual FStatusRecievedEvent& OnSelectedSessionStatusReceived() = 0;
 
+	/**
+	* Session selection changed event.
+	*/
+	DECLARE_EVENT(ITraceController, FSessionSelectionChanged);
+
+	/**
+	 * Event triggered when the session selection changes.
+	 */
+	virtual FSessionSelectionChanged& OnSessionSelectionChanged() = 0;
+
 	/**
 	 * Return true if a selected instance exists and has been discovered.
 	 */
@@ -178,4 +188,9 @@ public:
 	 * @param Func Functor to execute
 	 */
 	virtual void WithInstance(FGuid InstanceId, FCallback Func) = 0;
-};
+
+	/**
+	 * Returns the number of selected sessions.
+	 */
+	virtual uint32 GetNumSelectedInstances() = 0;
+};
\ No newline at end of file
