Trace Control Window now uses a slim toolbar style

diff --git a/Engine/Source/Developer/TraceTools/Private/TraceToolsStyle.cpp b/Engine/Source/Developer/TraceTools/Private/TraceToolsStyle.cpp
--- a/Engine/Source/Developer/TraceTools/Private/TraceToolsStyle.cpp
+++ b/Engine/Source/Developer/TraceTools/Private/TraceToolsStyle.cpp
@@ -6,11 +6,12 @@
 #include "Brushes/SlateImageBrush.h"
 #include "Brushes/SlateNoResource.h"
 #include "Misc/Paths.h"
+#include "Styling/CoreStyle.h"
 #include "Styling/StyleColors.h"
 #include "Styling/SlateStyleRegistry.h"
 #include "Styling/SlateTypes.h"
-#include "Styling/CoreStyle.h"
 #include "Styling/SlateStyleMacros.h"
+#include "Styling/ToolBarStyle.h"
 
 namespace UE::TraceTools
 {
@@ -116,7 +117,31 @@ void FTraceToolsStyle::Initialize()
 		.SetColorAndOpacity(FLinearColor(1.0f, 1.0f, 1.0f, 0.9f))
 		.SetShadowOffset(FVector2D(1, 1))
 		.SetShadowColorAndOpacity(FLinearColor(0, 0, 0, 0.9f)));
-	
+
+	// TraceControlToolbar
+	{
+		FToolBarStyle TraceControlToolbarStyle = FAppStyle::Get().GetWidgetStyle<FToolBarStyle>("SlimToolBar");
+
+		TraceControlToolbarStyle.SetBackgroundPadding(FMargin(4.0f, 4.0f));
+		TraceControlToolbarStyle.SetBlockPadding(FMargin(2.0f, 2.0f));
+		TraceControlToolbarStyle.SetButtonPadding(FMargin(2.0f, 2.0f));
+		TraceControlToolbarStyle.SetCheckBoxPadding(FMargin(2.0f, 2.0f));
+		TraceControlToolbarStyle.SetComboButtonPadding(FMargin(2.0f, 2.0f));
+		TraceControlToolbarStyle.SetIndentedBlockPadding(FMargin(2.0f, 2.0f));
+		TraceControlToolbarStyle.SetLabelPadding(FMargin(2.0f, 2.0f));
+		
+		TraceControlToolbarStyle.ToggleButton.SetPadding(FMargin(2.0f, 2.0f));
+
+		TraceControlToolbarStyle.ButtonStyle.SetNormalPadding(FMargin(6.0f, 2.0f, 4.0f, 2.0f));
+		TraceControlToolbarStyle.ButtonStyle.SetPressedPadding(FMargin(6.0f, 2.0f, 4.0f, 2.0f));
+
+		TraceControlToolbarStyle.IconSize.Set(20.0f, 20.0f);
+
+		StyleSet->Set("TraceControlToolbar", TraceControlToolbarStyle);
+
+		StyleSet->Set("TraceControlToolbar.MinUniformToolbarSize", 40.0f);
+		StyleSet->Set("TraceControlToolbar.MaxUniformToolbarSize", 40.0f);
+	}
 
 	// Filter list
 	/* Set images for various SCheckBox states associated with "ContentBrowser.FilterButton" ... */
diff --git a/Engine/Source/Developer/TraceTools/Private/Widgets/STraceControl.cpp b/Engine/Source/Developer/TraceTools/Private/Widgets/STraceControl.cpp
--- a/Engine/Source/Developer/TraceTools/Private/Widgets/STraceControl.cpp
+++ b/Engine/Source/Developer/TraceTools/Private/Widgets/STraceControl.cpp
@@ -34,6 +34,7 @@ void STraceControl::Construct(const FArguments& InArgs, TSharedPtr<ITraceControl
 		SNew(SVerticalBox)
 
 		+ SVerticalBox::Slot()
+		.Padding(0.0f, 4.0f, 0.0f, 0.0f)
 		.AutoHeight()
 		[
 			SNew(STraceControlToolbar, UICommandList.ToSharedRef(), InTraceController)
