Improvements to the trace statistics from trace control window.

diff --git a/Engine/Source/Developer/TraceTools/Private/Services/ISessionTraceFilterService.h b/Engine/Source/Developer/TraceTools/Private/Services/ISessionTraceFilterService.h
--- a/Engine/Source/Developer/TraceTools/Private/Services/ISessionTraceFilterService.h
+++ b/Engine/Source/Developer/TraceTools/Private/Services/ISessionTraceFilterService.h
@@ -28,6 +28,15 @@ struct FTraceObjectInfo
 	}
 };
 
+struct FTraceStats
+{
+	FTraceStatus::FStats StandardStats;
+
+	//Computed stats
+	uint64 BytesSentPerSecond = 0;
+	uint64 BytesTracedPerSecond = 0;
+};
+
 /** Filtering service, representing the state and data for a specific TraceServices::IAnalysisSession */
 class ISessionTraceFilterService : public TSharedFromThis<ISessionTraceFilterService>
 {
@@ -58,8 +67,8 @@ public:
 	/** Returns true if stats are available for the selected session. */
 	virtual bool HasStats() const = 0;
 
-	/** Get the settings of the selected session */
-	virtual const FTraceStatus::FStats& GetStats() const = 0;
+	/** Get the stats of the selected session */
+	virtual const FTraceStats& GetStats() const = 0;
 };
 
 } // namespace UE::TraceTools
\ No newline at end of file
diff --git a/Engine/Source/Developer/TraceTools/Private/Services/SessionTraceControllerFilterService.cpp b/Engine/Source/Developer/TraceTools/Private/Services/SessionTraceControllerFilterService.cpp
--- a/Engine/Source/Developer/TraceTools/Private/Services/SessionTraceControllerFilterService.cpp
+++ b/Engine/Source/Developer/TraceTools/Private/Services/SessionTraceControllerFilterService.cpp
@@ -104,6 +104,7 @@ void FSessionTraceControllerFilterService::DisableAllChannels()
 
 void FSessionTraceControllerFilterService::OnTraceStatusUpdated(const FTraceStatus& InStatus, FTraceStatus::EUpdateType InUpdateType, ITraceControllerCommands& Commands)
 {
+	double DeltaTimeSeconds = (FDateTime::Now() - TimeStamp).GetTotalSeconds();
 	if (!TraceController->HasAvailableSelectedInstance())
 	{
 		TimeStamp = FDateTime::Now();
@@ -124,7 +125,19 @@ void FSessionTraceControllerFilterService::OnTraceStatusUpdated(const FTraceStat
 
 	if (EnumHasAnyFlags(InUpdateType, FTraceStatus::EUpdateType::Status))
 	{
-		Stats = InStatus.Stats;
+		Stats.BytesSentPerSecond = 0;
+		Stats.BytesTracedPerSecond = 0;
+
+		if (DeltaTimeSeconds > 0.0f)
+		{
+			Stats.BytesSentPerSecond = FMath::Max(0ull, InStatus.Stats.BytesSent - Stats.StandardStats.BytesSent);
+			Stats.BytesSentPerSecond = (uint64) ((double)Stats.BytesSentPerSecond / DeltaTimeSeconds);
+
+			Stats.BytesTracedPerSecond = FMath::Max(0ull, InStatus.Stats.BytesTraced - Stats.StandardStats.BytesTraced);
+			Stats.BytesTracedPerSecond = (uint64) ((double)Stats.BytesTracedPerSecond / DeltaTimeSeconds);
+		}
+
+		Stats.StandardStats = InStatus.Stats;
 		bHasStats = true;
 	}
 }
@@ -186,7 +199,7 @@ bool FSessionTraceControllerFilterService::HasStats() const
 	return bHasStats;
 }
 
-const FTraceStatus::FStats& FSessionTraceControllerFilterService::GetStats() const
+const FTraceStats& FSessionTraceControllerFilterService::GetStats() const
 {
 	return Stats;
 }
diff --git a/Engine/Source/Developer/TraceTools/Private/Services/SessionTraceControllerFilterService.h b/Engine/Source/Developer/TraceTools/Private/Services/SessionTraceControllerFilterService.h
--- a/Engine/Source/Developer/TraceTools/Private/Services/SessionTraceControllerFilterService.h
+++ b/Engine/Source/Developer/TraceTools/Private/Services/SessionTraceControllerFilterService.h
@@ -35,7 +35,7 @@ public:
 	virtual const FTraceStatus::FSettings& GetSettings() const override;
 
 	virtual bool HasStats() const override;
-	virtual const FTraceStatus::FStats& GetStats() const override;
+	virtual const FTraceStats& GetStats() const override;
 	/** End ISessionTraceFilterService overrides */
 
 protected:
@@ -69,7 +69,7 @@ protected:
 	FTraceStatus::FSettings Settings;
 	bool bHasSettings = false;
 
-	FTraceStatus::FStats Stats;
+	FTraceStats Stats;
 	bool bHasStats = false;
 };
 
diff --git a/Engine/Source/Developer/TraceTools/Private/Widgets/STraceStatistics.cpp b/Engine/Source/Developer/TraceTools/Private/Widgets/STraceStatistics.cpp
--- a/Engine/Source/Developer/TraceTools/Private/Widgets/STraceStatistics.cpp
+++ b/Engine/Source/Developer/TraceTools/Private/Widgets/STraceStatistics.cpp
@@ -45,6 +45,7 @@ void STraceStatistics::Construct(const FArguments& InArgs, TSharedPtr<ISessionTr
 				SNew(SHorizontalBox)
 
 				+ SHorizontalBox::Slot()
+				.AutoWidth()
 				[
 					SNew(SVerticalBox)
 
@@ -71,6 +72,7 @@ void STraceStatistics::Construct(const FArguments& InArgs, TSharedPtr<ISessionTr
 						[
 							SNew(STextBlock)
 							.ColorAndOpacity(FSlateColor(EStyleColor::Foreground))
+							.ToolTipText(LOCTEXT("ImportantEventsTooltip", "The state of the Important Events cache."))
 							.Text(LOCTEXT("ImportantCache", "Important Events Cache:"))
 						]
 
@@ -96,6 +98,7 @@ void STraceStatistics::Construct(const FArguments& InArgs, TSharedPtr<ISessionTr
 						[
 							SNew(STextBlock)
 							.ColorAndOpacity(FSlateColor(EStyleColor::Foreground))
+							.ToolTipText(LOCTEXT("UseWorkerThreadTooltip", "If trace uses a worker thread. If not, TraceLog is pumped on end frame."))
 							.Text(LOCTEXT("WorkerThread", "Worker Thread:"))
 						]
 
@@ -121,7 +124,8 @@ void STraceStatistics::Construct(const FArguments& InArgs, TSharedPtr<ISessionTr
 						[
 							SNew(STextBlock)
 							.ColorAndOpacity(FSlateColor(EStyleColor::Foreground))
-							.Text(LOCTEXT("Tail Size", "Tail Size:"))
+							.ToolTipText(LOCTEXT("TailSizeTooltip", "Size of the tail buffer where the last seconds of trace data are stored."))
+							.Text(LOCTEXT("TailSize", "Tail Size:"))
 						]
 
 						+ SHorizontalBox::Slot()
@@ -136,6 +140,7 @@ void STraceStatistics::Construct(const FArguments& InArgs, TSharedPtr<ISessionTr
 				]
 
 				+ SHorizontalBox::Slot()
+				.AutoWidth()
 				.Padding(30.0f, 0.0f, 0.0f, 0.0f)
 				[
 					SNew(SVerticalBox)
@@ -161,41 +166,29 @@ void STraceStatistics::Construct(const FArguments& InArgs, TSharedPtr<ISessionTr
 						[
 							SNew(STextBlock)
 							.ColorAndOpacity(FSlateColor(EStyleColor::Foreground))
-							.Text(LOCTEXT("Bytes Sent", "Bytes Sent:"))
+							.ToolTipText(LOCTEXT("BytesSentTooltip", "Number of bytes sent to server or file."))
+							.Text(LOCTEXT("BytesSent", "Bytes Sent:"))
 						]
 
 						+ SHorizontalBox::Slot()
+						.AutoWidth()
 						.Padding(2.0f, 2.0, 0.0f, 0.0f)
 						[
 							SNew(STextBlock)
 							.ColorAndOpacity(FSlateColor(EStyleColor::AccentGray))
-							.Text_Lambda([this]() { return this->GetStatsMemoryValueText(SessionFilterService->GetStats().BytesSent); })
+							.Text_Lambda([this]() { return this->GetStatsMemoryValueText(SessionFilterService->GetStats().StandardStats.BytesSent); })
 						]
-					]
 
-					+ SVerticalBox::Slot()
-					[
-						SNew(SHorizontalBox)
-							
 						+ SHorizontalBox::Slot()
-						.Padding(0.0f, 2.0, 0.0f, 0.0f)
 						.AutoWidth()
-						[
-							SNew(STextBlock)
-							.ColorAndOpacity(FSlateColor(EStyleColor::Foreground))
-							.Text(LOCTEXT("Bytes Traced", "Bytes Traced:"))
-						]
-
-						+ SHorizontalBox::Slot()
-						.Padding(2.0f, 2.0, 0.0f, 0.0f)
+						.Padding(4.0f, 2.0, 0.0f, 0.0f)
 						[
 							SNew(STextBlock)
 							.ColorAndOpacity(FSlateColor(EStyleColor::AccentGray))
-							.Text_Lambda([this]() { return this->GetStatsMemoryValueText(SessionFilterService->GetStats().BytesTraced); })
+							.Text_Lambda([this]() { return this->GetStatsBandwidthText(SessionFilterService->GetStats().BytesSentPerSecond); })
 						]
 					]
 
-								
 					+ SVerticalBox::Slot()
 					[
 						SNew(SHorizontalBox)
@@ -206,40 +199,29 @@ void STraceStatistics::Construct(const FArguments& InArgs, TSharedPtr<ISessionTr
 						[
 							SNew(STextBlock)
 							.ColorAndOpacity(FSlateColor(EStyleColor::Foreground))
-							.Text(LOCTEXT("Memory Used", "Memory Used:"))
+							.ToolTipText(LOCTEXT("BytesTracedTooltip", "Number of (uncompressed) bytes traced from process."))
+							.Text(LOCTEXT("BytesTraced", "Bytes Traced:"))
 						]
 
 						+ SHorizontalBox::Slot()
+						.AutoWidth()
 						.Padding(2.0f, 2.0, 0.0f, 0.0f)
 						[
 							SNew(STextBlock)
 							.ColorAndOpacity(FSlateColor(EStyleColor::AccentGray))
-							.Text_Lambda([this]() { return this->GetStatsMemoryValueText(SessionFilterService->GetStats().MemoryUsed); })
+							.Text_Lambda([this]() { return this->GetStatsMemoryValueText(SessionFilterService->GetStats().StandardStats.BytesTraced); })
 						]
-					]
 
-					+ SVerticalBox::Slot()
-					[
-						SNew(SHorizontalBox)
-							
 						+ SHorizontalBox::Slot()
-						.Padding(0.0f, 2.0, 0.0f, 0.0f)
 						.AutoWidth()
-						[
-							SNew(STextBlock)
-							.ColorAndOpacity(FSlateColor(EStyleColor::Foreground))
-							.Text(LOCTEXT("Cache Allocated", "Cache Allocated:"))
-						]
-
-						+ SHorizontalBox::Slot()
 						.Padding(2.0f, 2.0, 0.0f, 0.0f)
 						[
 							SNew(STextBlock)
 							.ColorAndOpacity(FSlateColor(EStyleColor::AccentGray))
-							.Text_Lambda([this]() { return this->GetStatsMemoryValueText(SessionFilterService->GetStats().CacheAllocated); })
+							.Text_Lambda([this]() { return this->GetStatsBandwidthText(SessionFilterService->GetStats().BytesTracedPerSecond); })
 						]
 					]
-
+		
 					+ SVerticalBox::Slot()
 					[
 						SNew(SHorizontalBox)
@@ -250,7 +232,8 @@ void STraceStatistics::Construct(const FArguments& InArgs, TSharedPtr<ISessionTr
 						[
 							SNew(STextBlock)
 							.ColorAndOpacity(FSlateColor(EStyleColor::Foreground))
-							.Text(LOCTEXT("Cache Used", "Cache Used:"))
+							.ToolTipText(LOCTEXT("MemoryUsedTooltip", "Total memory used by TraceLog."))
+							.Text(LOCTEXT("MemoryUsed", "Memory Used:"))
 						]
 
 						+ SHorizontalBox::Slot()
@@ -258,7 +241,7 @@ void STraceStatistics::Construct(const FArguments& InArgs, TSharedPtr<ISessionTr
 						[
 							SNew(STextBlock)
 							.ColorAndOpacity(FSlateColor(EStyleColor::AccentGray))
-							.Text_Lambda([this]() { return this->GetStatsMemoryValueText(SessionFilterService->GetStats().CacheUsed); })
+							.Text_Lambda([this]() { return this->GetStatsMemoryValueText(SessionFilterService->GetStats().StandardStats.MemoryUsed); })
 						]
 					]
 
@@ -272,15 +255,17 @@ void STraceStatistics::Construct(const FArguments& InArgs, TSharedPtr<ISessionTr
 						[
 							SNew(STextBlock)
 							.ColorAndOpacity(FSlateColor(EStyleColor::Foreground))
-							.Text(LOCTEXT("Cache Waste", "Cache Waste:"))
+							.ToolTipText(LOCTEXT("ImportantEventsTooltip", "Memory for important events."))
+							.Text(LOCTEXT("ImportantEventsCache:", "Cache:"))
 						]
 
 						+ SHorizontalBox::Slot()
+						.AutoWidth()
 						.Padding(2.0f, 2.0, 0.0f, 0.0f)
 						[
 							SNew(STextBlock)
 							.ColorAndOpacity(FSlateColor(EStyleColor::AccentGray))
-							.Text_Lambda([this]() { return this->GetStatsMemoryValueText(SessionFilterService->GetStats().CacheWaste); })
+							.Text_Lambda([this]() { return this->GetStatsCacheText(); })
 						]
 					]
 				]
@@ -312,7 +297,12 @@ FText STraceStatistics::GetSettingsMemoryValueText(uint64 InValue) const
 		return LOCTEXT("N/A", "N/A");
 	}
 
-	return FText::AsMemory(InValue);
+	FNumberFormattingOptions FormattingOptionsMem;
+	FormattingOptionsMem.MaximumFractionalDigits = 2;
+	FormattingOptionsMem.MinimumFractionalDigits = 2;
+	FormattingOptionsMem.MinimumIntegralDigits = 1;
+
+	return FText::AsMemory(InValue, &FormattingOptionsMem);
 }
 
 FText STraceStatistics::GetStatsMemoryValueText(uint64 InValue) const
@@ -322,7 +312,50 @@ FText STraceStatistics::GetStatsMemoryValueText(uint64 InValue) const
 		return LOCTEXT("N/A", "N/A");
 	}
 
-	return FText::AsMemory(InValue);
+	FNumberFormattingOptions FormattingOptionsMem;
+	FormattingOptionsMem.MaximumFractionalDigits = 2;
+	FormattingOptionsMem.MinimumFractionalDigits = 2;
+	FormattingOptionsMem.MinimumIntegralDigits = 1;
+
+	return FText::AsMemory(InValue, &FormattingOptionsMem);
+}
+
+FText STraceStatistics::GetStatsBandwidthText(uint64 InValue) const
+{
+	if (!SessionFilterService->HasStats())
+	{
+		return FText();
+	}
+
+	FNumberFormattingOptions FormattingOptionsMem;
+	FormattingOptionsMem.MaximumFractionalDigits = 2;
+	FormattingOptionsMem.MinimumFractionalDigits = 2;
+	FormattingOptionsMem.MinimumIntegralDigits = 1;
+
+	FText Result = FText::AsMemory(InValue, &FormattingOptionsMem);
+
+	return FText::Format(LOCTEXT("TraceStatBandwidthFormat", "({0}/s)"), Result);
+}
+
+FText STraceStatistics::GetStatsCacheText() const
+{
+	if (!SessionFilterService->HasStats())
+	{
+		return FText();
+	}
+
+	FNumberFormattingOptions FormattingOptionsMem;
+	FormattingOptionsMem.MaximumFractionalDigits = 2;
+	FormattingOptionsMem.MinimumFractionalDigits = 2;
+	FormattingOptionsMem.MinimumIntegralDigits = 1;
+
+	const FTraceStatus::FStats& Stats = SessionFilterService->GetStats().StandardStats;
+	FText CacheAllocated = FText::AsMemory(Stats.CacheAllocated, &FormattingOptionsMem);
+	FText CacheUsed = FText::AsMemory(Stats.CacheUsed, &FormattingOptionsMem);
+	FText CacheUnused = FText::AsMemory(Stats.CacheAllocated - Stats.CacheUsed, &FormattingOptionsMem);
+	FText CacheWasted = FText::AsMemory(Stats.CacheWaste, &FormattingOptionsMem);
+
+	return FText::Format(LOCTEXT("TraceCacheTextFormat", "{0} ({1} used + {2} unused | {3} waste)"), CacheAllocated, CacheUsed, CacheUnused, CacheWasted);
 }
 
 } // namespace UE::TraceTools
diff --git a/Engine/Source/Developer/TraceTools/Private/Widgets/STraceStatistics.h b/Engine/Source/Developer/TraceTools/Private/Widgets/STraceStatistics.h
--- a/Engine/Source/Developer/TraceTools/Private/Widgets/STraceStatistics.h
+++ b/Engine/Source/Developer/TraceTools/Private/Widgets/STraceStatistics.h
@@ -33,6 +33,8 @@ private:
 	FText GetSettingsOnOffText(bool InValue) const;
 	FText GetSettingsMemoryValueText(uint64 InValue) const;
 	FText GetStatsMemoryValueText(uint64 InValue) const;
+	FText GetStatsBandwidthText(uint64 InValue) const;
+	FText GetStatsCacheText() const;
 
 private:
 	TSharedPtr<ISessionTraceFilterService> SessionFilterService;
