Fix loctext key conflict

diff --git a/Engine/Source/Developer/TraceTools/Private/Widgets/STraceStatistics.cpp b/Engine/Source/Developer/TraceTools/Private/Widgets/STraceStatistics.cpp
--- a/Engine/Source/Developer/TraceTools/Private/Widgets/STraceStatistics.cpp
+++ b/Engine/Source/Developer/TraceTools/Private/Widgets/STraceStatistics.cpp
@@ -72,7 +72,7 @@ void STraceStatistics::Construct(const FArguments& InArgs, TSharedPtr<ISessionTr
 						[
 							SNew(STextBlock)
 							.ColorAndOpacity(FSlateColor(EStyleColor::Foreground))
-							.ToolTipText(LOCTEXT("ImportantEventsTooltip", "The state of the Important Events cache."))
+							.ToolTipText(LOCTEXT("ImportantEventsSettingTooltip", "The state of the Important Events cache."))
 							.Text(LOCTEXT("ImportantCache", "Important Events Cache:"))
 						]
 
@@ -255,7 +255,7 @@ void STraceStatistics::Construct(const FArguments& InArgs, TSharedPtr<ISessionTr
 						[
 							SNew(STextBlock)
 							.ColorAndOpacity(FSlateColor(EStyleColor::Foreground))
-							.ToolTipText(LOCTEXT("ImportantEventsTooltip", "Memory for important events."))
+							.ToolTipText(LOCTEXT("ImportantEventsMemoryTooltip", "Memory for important events."))
 							.Text(LOCTEXT("ImportantEventsCache:", "Cache:"))
 						]
 
