Create a widget for a dynamic start/stop trace button.
Add the widget to the trace control widget.
Use the widget in the Insights Status Bar widget also.

diff --git a/Engine/Plugins/TraceUtilities/Source/EditorTraceUtilities/EditorTraceUtilities.Build.cs b/Engine/Plugins/TraceUtilities/Source/EditorTraceUtilities/EditorTraceUtilities.Build.cs
--- a/Engine/Plugins/TraceUtilities/Source/EditorTraceUtilities/EditorTraceUtilities.Build.cs
+++ b/Engine/Plugins/TraceUtilities/Source/EditorTraceUtilities/EditorTraceUtilities.Build.cs
@@ -31,6 +31,7 @@ public class EditorTraceUtilities : ModuleRules
 				"ToolMenus",
 				"TraceAnalysis",
 				"TraceInsights",
+				"TraceTools",
 				"TraceLog",
 				"UATHelper",
 			}
diff --git a/Engine/Plugins/TraceUtilities/Source/EditorTraceUtilities/Private/InsightsStatusBar/EditorTraceUtilitiesStyle.h b/Engine/Plugins/TraceUtilities/Source/EditorTraceUtilities/Private/InsightsStatusBar/EditorTraceUtilitiesStyle.h
--- a/Engine/Plugins/TraceUtilities/Source/EditorTraceUtilities/Private/InsightsStatusBar/EditorTraceUtilitiesStyle.h
+++ b/Engine/Plugins/TraceUtilities/Source/EditorTraceUtilities/Private/InsightsStatusBar/EditorTraceUtilitiesStyle.h
@@ -31,12 +31,6 @@ public:
 		Set("Icons.Trace.Menu", new IMAGE_BRUSH_SVG("Trace_16", Icon16x16));
 		Set("Icons.Trace.StatusBar", new IMAGE_BRUSH_SVG("Trace_20", Icon16x16));
 
-		Set("Icons.RecordTraceCenter.StatusBar", new IMAGE_BRUSH_SVG("RecordTraceCenter_20", Icon16x16));
-		Set("Icons.RecordTraceOutline.StatusBar", new IMAGE_BRUSH_SVG("RecordTraceOutline_20", Icon16x16));
-		Set("Icons.RecordTraceRecording.StatusBar", new IMAGE_BRUSH_SVG("RecordTraceRecording_20", Icon16x16));
-		
-		Set("Icons.RecordTraceStop.StatusBar", new IMAGE_BRUSH_SVG("RecordTraceStop_20", Icon16x16, FStyleColors::Error));
-
 		Set("Icons.StartTrace.Menu", new IMAGE_BRUSH_SVG("StartTrace_16", Icon16x16));
 		Set("Icons.StartTrace.StatusBar", new IMAGE_BRUSH_SVG("StartTrace_20", Icon16x16));
 
diff --git a/Engine/Plugins/TraceUtilities/Source/EditorTraceUtilities/Private/InsightsStatusBar/SInsightsStatusBar.cpp b/Engine/Plugins/TraceUtilities/Source/EditorTraceUtilities/Private/InsightsStatusBar/SInsightsStatusBar.cpp
--- a/Engine/Plugins/TraceUtilities/Source/EditorTraceUtilities/Private/InsightsStatusBar/SInsightsStatusBar.cpp
+++ b/Engine/Plugins/TraceUtilities/Source/EditorTraceUtilities/Private/InsightsStatusBar/SInsightsStatusBar.cpp
@@ -26,13 +26,14 @@
 #include "SRecentTracesList.h"
 #include "Styling/StyleColors.h"
 #include "ToolMenus.h"
+#include "TraceTools/Interfaces/ITraceToolsModule.h"
+#include "TraceTools/Widgets/SToggleTraceButton.h"
 #include "Trace/Detail/Channel.h"
 #include "Trace/StoreClient.h"
 #include "Trace/Trace.h"
 #include "UnrealInsightsLauncher.h"
 #include "Insights/Widgets/STraceServerControl.h"
 #include "Widgets/SBoxPanel.h"
-#include "Widgets/SOverlay.h"
 #include "Widgets/Images/SImage.h"
 #include "Widgets/Input/SButton.h"
 #include "Widgets/Input/SComboButton.h"
@@ -128,6 +129,8 @@ void SInsightsStatusBarWidget::Tick(const FGeometry& AllottedGeometry, const dou
 
 void SInsightsStatusBarWidget::Construct(const FArguments& InArgs)
 {
+	FModuleManager::LoadModuleChecked<UE::TraceTools::ITraceToolsModule>("TraceTools");
+
 	this->ChildSlot
 	[
 		SNew(SHorizontalBox)
@@ -168,42 +171,11 @@ void SInsightsStatusBarWidget::Construct(const FArguments& InArgs)
 		+ SHorizontalBox::Slot()
 		.AutoWidth()
 		[
-			SNew(SButton)
-			.ButtonStyle(FAppStyle::Get(), "SimpleButton")
-			.ContentPadding(FMargin(0.0f, 0.0f, 0.0f, 3.0f))
-			.HAlign(HAlign_Left)
-			.VAlign(VAlign_Bottom)
-			.ToolTipText(this, &SInsightsStatusBarWidget::GetRecordingButtonTooltipText)
-			.OnClicked_Lambda([this]() { this->ToggleTrace_OnClicked(); return FReply::Handled(); })
-			.OnHovered_Lambda([this]() { this->bIsTraceRecordButtonHovered = true; })
-			.OnUnhovered_Lambda([this]() { this->bIsTraceRecordButtonHovered = false; })
-			.Content()
-			[
-				SNew(SOverlay)
-
-				+ SOverlay::Slot()
-				[
-					SNew(SImage)
-					.ColorAndOpacity(this, &SInsightsStatusBarWidget::GetRecordingButtonColor)
-					.Image(FEditorTraceUtilitiesStyle::Get().GetBrush("Icons.RecordTraceCenter.StatusBar"))
-					.Visibility(this, &SInsightsStatusBarWidget::GetStartTraceIconVisibility)
-				]
-
-				+ SOverlay::Slot()
-				[
-					SNew(SImage)
-					.ColorAndOpacity(this, &SInsightsStatusBarWidget::GetRecordingButtonOutlineColor)
-					.Image(FEditorTraceUtilitiesStyle::Get().GetBrush("Icons.RecordTraceOutline.StatusBar"))
-					.Visibility(this, &SInsightsStatusBarWidget::GetStartTraceIconVisibility)
-				]
-
-				+ SOverlay::Slot()
-				[
-					SNew(SImage)
-					.Image(FEditorTraceUtilitiesStyle::Get().GetBrush("Icons.RecordTraceStop.StatusBar"))
-					.Visibility(this, &SInsightsStatusBarWidget::GetStopTraceIconVisibility)
-				]
-			]
+			SNew(UE::TraceTools::SToggleTraceButton)
+			.OnToggleTraceRequested(this, &SInsightsStatusBarWidget::ToggleTrace_OnClicked)
+			.IsTraceRunning_Lambda([]() {return UE::Trace::IsTracing(); })
+			.ButtonSize(UE::TraceTools::SToggleTraceButton::EButtonSize::StatusBar)
+			.IsEnabled(this, &SInsightsStatusBarWidget::ToggleTrace_CanExecute)
 		]
 
 		+ SHorizontalBox::Slot()
@@ -558,38 +530,6 @@ FText SInsightsStatusBarWidget::GetTitleToolTipText() const
 	return DescBuilder.ToText();
 }
 
-FSlateColor SInsightsStatusBarWidget::GetRecordingButtonColor() const
-{
-	if (!UE::Trace::IsTracing())
-	{
-		return FStyleColors::White;
-	}
-
-	return FStyleColors::Error;
-}
-
-FSlateColor SInsightsStatusBarWidget::GetRecordingButtonOutlineColor() const
-{
-	if (!UE::Trace::IsTracing())
-	{
-		ConnectionStartTime = FSlateApplication::Get().GetCurrentTime();
-		return FLinearColor::White.CopyWithNewOpacity(0.5f);
-	}
-
-	double ElapsedTime = FSlateApplication::Get().GetCurrentTime() - ConnectionStartTime;
-	return FStyleColors::Error.GetColor(FWidgetStyle()).CopyWithNewOpacity(0.5f + 0.5f * FMath::MakePulsatingValue(ElapsedTime, 0.5f));
-}
-
-FText SInsightsStatusBarWidget::GetRecordingButtonTooltipText() const
-{
-	if (!UE::Trace::IsTracing())
-	{
-		return LOCTEXT("StartTracing", "Start tracing. The trace destination is set from the menu.");
-	}
-
-	return LOCTEXT("StopTracing", "Stop Tracing.");
-}
-
 void SInsightsStatusBarWidget::LogMessage(const FText& Text)
 {
 	FMessageLog ReportMessageLog(LogListingName);
@@ -773,6 +713,15 @@ FText SInsightsStatusBarWidget::GetTraceMenuItemTooltipText() const
 	return LOCTEXT("StartTraceButtonTooltip", "Start tracing to the selected trace destination.");
 }
 
+bool SInsightsStatusBarWidget::ToggleTrace_CanExecute() const
+{
+#if UE_TRACE_ENABLED
+	return true;
+#else
+	return false;
+#endif
+}
+
 void SInsightsStatusBarWidget::ToggleTrace_OnClicked()
 {
 	if (UE::Trace::IsTracing())
@@ -849,26 +798,6 @@ bool SInsightsStatusBarWidget::StartTracing()
 	return false;
 }
 
-EVisibility SInsightsStatusBarWidget::GetStartTraceIconVisibility() const
-{
-	if (GetStopTraceIconVisibility() == EVisibility::Hidden)
-	{
-		return EVisibility::Visible;
-	}
-
-	return EVisibility::Hidden;
-}
-
-EVisibility SInsightsStatusBarWidget::GetStopTraceIconVisibility() const
-{
-	if (bIsTraceRecordButtonHovered && UE::Trace::IsTracing())
-	{
-		return EVisibility::Visible;
-	}
-
-	return EVisibility::Hidden;
-}
-
 bool SInsightsStatusBarWidget::GetBooleanSettingValue(const TCHAR* InSettingName)
 {
 	bool Value = false;
diff --git a/Engine/Plugins/TraceUtilities/Source/EditorTraceUtilities/Private/InsightsStatusBar/SInsightsStatusBar.h b/Engine/Plugins/TraceUtilities/Source/EditorTraceUtilities/Private/InsightsStatusBar/SInsightsStatusBar.h
--- a/Engine/Plugins/TraceUtilities/Source/EditorTraceUtilities/Private/InsightsStatusBar/SInsightsStatusBar.h
+++ b/Engine/Plugins/TraceUtilities/Source/EditorTraceUtilities/Private/InsightsStatusBar/SInsightsStatusBar.h
@@ -64,10 +64,6 @@ public:
 private:
 	FText	GetTitleToolTipText() const;
 	
-	FSlateColor GetRecordingButtonColor() const;
-	FSlateColor GetRecordingButtonOutlineColor() const;
-	FText GetRecordingButtonTooltipText() const;
-
 	void LaunchUnrealInsights_OnClicked();
 	
 	void OpenLiveSession_OnClicked();
@@ -91,15 +87,14 @@ private:
 
 	FText GetTraceMenuItemText() const;
 	FText GetTraceMenuItemTooltipText() const;
+
 	void ToggleTrace_OnClicked();
+	bool ToggleTrace_CanExecute() const;
 
 	bool PauseTrace_CanExecute();
 	FText GetPauseTraceMenuItemTooltipText() const;
 	void TogglePauseTrace_OnClicked();
 
-	EVisibility GetStartTraceIconVisibility() const;
-	EVisibility GetStopTraceIconVisibility() const;
-
 	bool StartTracing();
 
 	TSharedRef<SWidget> MakeTraceMenu();
diff --git a/Engine/Source/Developer/SessionFrontend/Private/Widgets/SSessionFrontend.cpp b/Engine/Source/Developer/SessionFrontend/Private/Widgets/SSessionFrontend.cpp
--- a/Engine/Source/Developer/SessionFrontend/Private/Widgets/SSessionFrontend.cpp
+++ b/Engine/Source/Developer/SessionFrontend/Private/Widgets/SSessionFrontend.cpp
@@ -15,9 +15,9 @@
 #include "IAutomationWindowModule.h"
 #include "Interfaces/IScreenShotToolsModule.h"
 #include "Interfaces/IScreenShotComparisonModule.h"
-#include "Interfaces/ITraceToolsModule.h"
 #include "ISessionServicesModule.h"
 #include "SessionFrontendStyle.h"
+#include "TraceTools/Interfaces/ITraceToolsModule.h"
 #include "Widgets/Browser/SSessionBrowser.h"
 #include "Widgets/Console/SSessionConsole.h"
 #include "Widgets/Docking/SDockTab.h"
diff --git a/Engine/Source/Developer/TraceTools/Private/TraceToolsModule.h b/Engine/Source/Developer/TraceTools/Private/TraceToolsModule.h
--- a/Engine/Source/Developer/TraceTools/Private/TraceToolsModule.h
+++ b/Engine/Source/Developer/TraceTools/Private/TraceToolsModule.h
@@ -2,7 +2,7 @@
 
 #pragma once
 
-#include "Interfaces/ITraceToolsModule.h"
+#include "TraceTools/Interfaces/ITraceToolsModule.h"
 
 #include "Modules/ModuleInterface.h"
 
diff --git a/Engine/Source/Developer/TraceTools/Private/TraceToolsStyle.cpp b/Engine/Source/Developer/TraceTools/Private/TraceToolsStyle.cpp
--- a/Engine/Source/Developer/TraceTools/Private/TraceToolsStyle.cpp
+++ b/Engine/Source/Developer/TraceTools/Private/TraceToolsStyle.cpp
@@ -168,9 +168,6 @@ void FTraceToolsStyle::Initialize()
 	StyleSet->Set("TraceControl.StartTrace", new CORE_IMAGE_BRUSH_SVG("Starship/Insights/TraceTools/TraceStart", Icon40x40));
 	StyleSet->Set("TraceControl.StartTrace.Small", new CORE_IMAGE_BRUSH_SVG("Starship/Insights/TraceTools/TraceStart", Icon20x20));
 
-	StyleSet->Set("TraceControl.StopTrace", new CORE_IMAGE_BRUSH_SVG("Starship/Insights/TraceTools/TraceStop", Icon40x40));
-	StyleSet->Set("TraceControl.StopTrace.Small", new CORE_IMAGE_BRUSH_SVG("Starship/Insights/TraceTools/TraceStop", Icon20x20));
-
 	StyleSet->Set("TraceControl.PauseTrace", new CORE_IMAGE_BRUSH_SVG("Starship/Insights/TraceTools/TracePause", Icon40x40));
 	StyleSet->Set("TraceControl.PauseTrace.Small", new CORE_IMAGE_BRUSH_SVG("Starship/Insights/TraceTools/TracePause", Icon20x20));
 
@@ -194,6 +191,16 @@ void FTraceToolsStyle::Initialize()
 
 	StyleSet->Set("TraceControl.ToggleStatNamedEvents", new CORE_IMAGE_BRUSH_SVG("../Editor/Slate/Starship/Common/Statistics", Icon40x40));
 	StyleSet->Set("TraceControl.ToggleStatNamedEvents.Small", new CORE_IMAGE_BRUSH_SVG("../Editor/Slate/Starship/Common/Statistics", Icon20x20));
+
+	StyleSet->Set("ToggleTraceButton.RecordTraceCenter.StatusBar", new CORE_IMAGE_BRUSH_SVG("Starship/Insights/TraceTools/RecordTraceCenter", Icon16x16));
+	StyleSet->Set("ToggleTraceButton.RecordTraceOutline.StatusBar", new CORE_IMAGE_BRUSH_SVG("Starship/Insights/TraceTools/RecordTraceOutline", Icon16x16));
+	StyleSet->Set("ToggleTraceButton.RecordTraceRecording.StatusBar", new CORE_IMAGE_BRUSH_SVG("Starship/Insights/TraceTools/RecordTraceRecording", Icon16x16));
+	StyleSet->Set("ToggleTraceButton.TraceStop.StatusBar", new CORE_IMAGE_BRUSH_SVG("Starship/Insights/TraceTools/TraceStop", Icon16x16, FStyleColors::Error));
+
+	StyleSet->Set("ToggleTraceButton.RecordTraceCenter.SlimToolbar", new CORE_IMAGE_BRUSH_SVG("Starship/Insights/TraceTools/RecordTraceCenter", Icon20x20));
+	StyleSet->Set("ToggleTraceButton.RecordTraceOutline.SlimToolbar", new CORE_IMAGE_BRUSH_SVG("Starship/Insights/TraceTools/RecordTraceOutline", Icon20x20));
+	StyleSet->Set("ToggleTraceButton.RecordTraceRecording.SlimToolbar", new CORE_IMAGE_BRUSH_SVG("Starship/Insights/TraceTools/RecordTraceRecording", Icon20x20));
+	StyleSet->Set("ToggleTraceButton.TraceStop.SlimToolbar", new CORE_IMAGE_BRUSH_SVG("Starship/Insights/TraceTools/TraceStop", Icon20x20, FStyleColors::Error));
 	
 	FSlateStyleRegistry::RegisterSlateStyle( *StyleSet.Get() );
 }
diff --git a/Engine/Source/Developer/TraceTools/Private/Widgets/SToggleTraceButton.cpp b/Engine/Source/Developer/TraceTools/Private/Widgets/SToggleTraceButton.cpp
new file mode 100644
--- /dev/null
+++ b/Engine/Source/Developer/TraceTools/Private/Widgets/SToggleTraceButton.cpp
@@ -0,0 +1,189 @@
+// Copyright Epic Games, Inc. All Rights Reserved.
+
+#include "TraceTools/Widgets/SToggleTraceButton.h"
+
+#include "Framework/Application/SlateApplication.h"
+#include "Internationalization/Text.h"
+#include "SlateOptMacros.h"
+#include "Styling/StyleColors.h"
+#include "Widgets/Images/SImage.h"
+#include "Widgets/Input/SButton.h"
+#include "Widgets/Layout/SBorder.h"
+#include "Widgets/SBoxPanel.h"
+#include "Widgets/SOverlay.h"
+#include "Widgets/SBoxPanel.h"
+#include "Widgets/Text/STextBlock.h"
+
+//TraceTools
+#include "Services/SessionTraceControllerFilterService.h"
+#include "TraceToolsStyle.h"
+
+#define LOCTEXT_NAMESPACE "SToggleTraceButton"
+
+namespace UE::TraceTools
+{
+
+SLATE_IMPLEMENT_WIDGET(SToggleTraceButton)
+
+void SToggleTraceButton::PrivateRegisterAttributes(FSlateAttributeInitializer& AttributeInitializer)
+{
+	SLATE_ADD_MEMBER_ATTRIBUTE_DEFINITION_WITH_NAME(AttributeInitializer, "IsTraceRunning", IsTraceRunningAttribute, EInvalidateWidgetReason::Paint);
+}
+
+SToggleTraceButton::SToggleTraceButton()
+	: IsTraceRunningAttribute(*this)
+{
+}
+
+SToggleTraceButton::~SToggleTraceButton()
+{
+}
+
+BEGIN_SLATE_FUNCTION_BUILD_OPTIMIZATION
+void SToggleTraceButton::Construct(const FArguments& InArgs)
+{
+	OnToggleTraceRequested = InArgs._OnToggleTraceRequested;
+	IsTraceRunningAttribute.Assign(*this, InArgs._IsTraceRunning);
+	ButtonSize = InArgs._ButtonSize;
+
+	ChildSlot
+	[
+		SNew(SButton)
+		.ButtonStyle(FAppStyle::Get(), "SimpleButton")
+		.ContentPadding(FMargin(0.0f, 0.0f, 0.0f, 3.0f))
+		.HAlign(HAlign_Left)
+		.VAlign(VAlign_Bottom)
+		.ToolTipText(this, &SToggleTraceButton::GetRecordingButtonTooltipText)
+		.OnClicked_Lambda([this]() { this->ToggleTrace_OnClicked(); return FReply::Handled(); })
+		.OnHovered_Lambda([this]() { this->bIsTraceRecordButtonHovered = true; })
+		.OnUnhovered_Lambda([this]() { this->bIsTraceRecordButtonHovered = false; })
+		.Content()
+		[
+			SNew(SOverlay)
+
+			+ SOverlay::Slot()
+			[
+				SNew(SImage)
+				.ColorAndOpacity(this, &SToggleTraceButton::GetRecordingButtonColor)
+				.Image(GetToggleTraceCenterBrush())
+				.Visibility(this, &SToggleTraceButton::GetStartTraceIconVisibility)
+			]
+
+			+ SOverlay::Slot()
+			[
+				SNew(SImage)
+				.ColorAndOpacity(this, &SToggleTraceButton::GetRecordingButtonOutlineColor)
+				.Image(GetToggleTraceOutlineBrush())
+				.Visibility(this, &SToggleTraceButton::GetStartTraceIconVisibility)
+			]
+
+			+ SOverlay::Slot()
+			[
+				SNew(SImage)
+				.Image(GetToggleTraceStopBrush())
+				.Visibility(this, &SToggleTraceButton::GetStopTraceIconVisibility)
+			]
+		]
+	];
+}
+END_SLATE_FUNCTION_BUILD_OPTIMIZATION
+
+EVisibility SToggleTraceButton::GetStartTraceIconVisibility() const
+{
+	if (GetStopTraceIconVisibility() == EVisibility::Hidden)
+	{
+		return EVisibility::Visible;
+	}
+
+	return EVisibility::Hidden;
+}
+
+EVisibility SToggleTraceButton::GetStopTraceIconVisibility() const
+{
+	if (bIsTraceRecordButtonHovered && IsTraceRunningAttribute.Get())
+	{
+		return EVisibility::Visible;
+	}
+
+	return EVisibility::Hidden;
+}
+
+FSlateColor SToggleTraceButton::GetRecordingButtonColor() const
+{
+	if (!IsTraceRunningAttribute.Get())
+	{
+		return FStyleColors::White;
+	}
+
+	return FStyleColors::Error;
+}
+
+FSlateColor SToggleTraceButton::GetRecordingButtonOutlineColor() const
+{
+	if (!IsTraceRunningAttribute.Get())
+	{
+		ConnectionStartTime = FSlateApplication::Get().GetCurrentTime();
+		return FLinearColor::White.CopyWithNewOpacity(0.5f);
+	}
+
+	double ElapsedTime = FSlateApplication::Get().GetCurrentTime() - ConnectionStartTime;
+	return FStyleColors::Error.GetColor(FWidgetStyle()).CopyWithNewOpacity(0.5f + 0.5f * FMath::MakePulsatingValue(ElapsedTime, 0.5f));
+}
+
+FText SToggleTraceButton::GetRecordingButtonTooltipText() const
+{
+	if (!IsTraceRunningAttribute.Get())
+	{
+		return LOCTEXT("StartTracing", "Start tracing. The trace destination is set from the menu.");
+	}
+
+	return LOCTEXT("StopTracing", "Stop Tracing.");
+}
+
+void SToggleTraceButton::ToggleTrace_OnClicked() const
+{
+	OnToggleTraceRequested.ExecuteIfBound();
+}
+
+const FSlateBrush* SToggleTraceButton::GetToggleTraceCenterBrush() const
+{
+	switch (ButtonSize)
+	{
+	case EButtonSize::StatusBar:
+		return FTraceToolsStyle::Get().GetBrush("ToggleTraceButton.RecordTraceCenter.StatusBar");
+	case EButtonSize::SlimToolbar:
+		return FTraceToolsStyle::Get().GetBrush("ToggleTraceButton.RecordTraceCenter.SlimToolbar");
+	}
+
+	return nullptr;
+}
+
+const FSlateBrush* SToggleTraceButton::GetToggleTraceOutlineBrush() const
+{
+	switch (ButtonSize)
+	{
+	case EButtonSize::StatusBar:
+		return FTraceToolsStyle::Get().GetBrush("ToggleTraceButton.RecordTraceOutline.StatusBar");
+	case EButtonSize::SlimToolbar:
+		return FTraceToolsStyle::Get().GetBrush("ToggleTraceButton.RecordTraceOutline.SlimToolbar");
+	}
+
+	return nullptr;
+}
+
+const FSlateBrush* SToggleTraceButton::GetToggleTraceStopBrush() const
+{
+	switch (ButtonSize)
+	{
+	case EButtonSize::StatusBar:
+		return FTraceToolsStyle::Get().GetBrush("ToggleTraceButton.TraceStop.StatusBar");
+	case EButtonSize::SlimToolbar:
+		return FTraceToolsStyle::Get().GetBrush("ToggleTraceButton.TraceStop.SlimToolbar");
+	}
+	
+	return nullptr;
+}
+
+} // namespace UE::TraceTools
+
+#undef LOCTEXT_NAMESPACE
\ No newline at end of file
diff --git a/Engine/Source/Developer/TraceTools/Private/Widgets/STraceControlToolbar.cpp b/Engine/Source/Developer/TraceTools/Private/Widgets/STraceControlToolbar.cpp
--- a/Engine/Source/Developer/TraceTools/Private/Widgets/STraceControlToolbar.cpp
+++ b/Engine/Source/Developer/TraceTools/Private/Widgets/STraceControlToolbar.cpp
@@ -9,9 +9,11 @@
 #include "SocketSubsystem.h"
 #include "Styling/AppStyle.h"
 #include "Widgets/Layout/SBorder.h"
+#include "Widgets/SBoxPanel.h"
 
 //TraceTools
 #include "Models/TraceControlCommands.h"
+#include "TraceTools/Widgets/SToggleTraceButton.h"
 
 #define LOCTEXT_NAMESPACE "STraceControlToolbar"
 
@@ -55,8 +57,22 @@ void STraceControlToolbar::Construct(const FArguments& InArgs, const TSharedRef<
 
 		Toolbar.SetLabelVisibility(EVisibility::Collapsed);
 
-		Toolbar.AddToolBarButton(FTraceControlCommands::Get().StartTrace);
-		Toolbar.AddToolBarButton(FTraceControlCommands::Get().StopTrace);
+		TSharedRef<SWidget> ToggleTraceWidget = 
+			SNew(SHorizontalBox)
+			
+			+ SHorizontalBox::Slot()
+			.Padding(0.0f, 2.0f, 0.0f, 0.0f)
+			[
+				SNew(SToggleTraceButton)
+				.OnToggleTraceRequested(this, &STraceControlToolbar::ToggleTrace_Execute)
+				.IsTraceRunning_Lambda([this]() {return bIsTracing; })
+				.IsEnabled(this, &STraceControlToolbar::ToggleTrace_CanExecute)
+				.ButtonSize(SToggleTraceButton::EButtonSize::SlimToolbar)
+			];
+
+
+		Toolbar.AddToolBarWidget(ToggleTraceWidget);
+
 		Toolbar.AddToolBarButton(FTraceControlCommands::Get().TraceSnapshot);
 
 		Toolbar.AddSeparator();
@@ -108,14 +124,6 @@ void STraceControlToolbar::BindCommands(const TSharedRef<FUICommandList>& Comman
 						   FExecuteAction::CreateSP(this, &STraceControlToolbar::SetTraceTarget_Execute, ETraceTarget::File),
 						   FCanExecuteAction::CreateSP(this, &STraceControlToolbar::SetTraceTarget_CanExecute));
 
-	CommandList->MapAction(FTraceControlCommands::Get().StartTrace, 
-					       FExecuteAction::CreateSP(this, &STraceControlToolbar::StartTrace_Execute),
-					       FCanExecuteAction::CreateSP(this, &STraceControlToolbar::StartTrace_CanExecute));
-
-	CommandList->MapAction(FTraceControlCommands::Get().StopTrace, 
-						   FExecuteAction::CreateSP(this, &STraceControlToolbar::StopTrace_Execute),
-						   FCanExecuteAction::CreateSP(this, &STraceControlToolbar::StopTrace_CanExecute));
-
 	CommandList->MapAction(FTraceControlCommands::Get().TraceSnapshot, 
 						   FExecuteAction::CreateSP(this, &STraceControlToolbar::TraceSnapshot_Execute),
 						   FCanExecuteAction::CreateSP(this, &STraceControlToolbar::TraceSnapshot_CanExecute));
@@ -171,39 +179,36 @@ void STraceControlToolbar::SetTraceTarget_Execute(ETraceTarget InTraceTarget)
 	TraceTarget = InTraceTarget;
 }
 
-bool STraceControlToolbar::StartTrace_CanExecute() const
+bool STraceControlToolbar::ToggleTrace_CanExecute() const
 {
-	return TraceController->HasAvailableSelectedInstance() && bIsTracing == false;
+	return TraceController->HasAvailableSelectedInstance();
 }
 
-void STraceControlToolbar::StartTrace_Execute()
+void STraceControlToolbar::ToggleTrace_Execute()
 {
-	TraceController->WithSelectedInstances([&](const FTraceStatus& Status, ITraceControllerCommands& Commands)
+	if (!bIsTracing)
 	{
-		if (TraceTarget == ETraceTarget::Server)
-		{
-			Commands.Send(TraceHostAddr, TEXT(""));
-		}
-		else if (TraceTarget == ETraceTarget::File)
-		{
-			Commands.File(TEXT(""), TEXT(""));
-		}
-	});
-	bIsTracing = true;
-}
-
-bool STraceControlToolbar::StopTrace_CanExecute() const
-{
-	return TraceController->HasAvailableSelectedInstance() && bIsTracing == true;
-}
-
-void STraceControlToolbar::StopTrace_Execute()
-{
-	TraceController->WithSelectedInstances([&](const FTraceStatus& Status, ITraceControllerCommands& Commands)
+		TraceController->WithSelectedInstances([&](const FTraceStatus& Status, ITraceControllerCommands& Commands)
+			{
+				if (TraceTarget == ETraceTarget::Server)
+				{
+					Commands.Send(TraceHostAddr, TEXT(""));
+				}
+				else if (TraceTarget == ETraceTarget::File)
+				{
+					Commands.File(TEXT(""), TEXT(""));
+				}
+			});
+		bIsTracing = true;
+	}
+	else
 	{
-		Commands.Stop();
-	});
-	bIsTracing = false;
+		TraceController->WithSelectedInstances([&](const FTraceStatus& Status, ITraceControllerCommands& Commands)
+			{
+				Commands.Stop();
+			});
+		bIsTracing = false;
+	}
 }
 
 bool STraceControlToolbar::TraceSnapshot_CanExecute() const
diff --git a/Engine/Source/Developer/TraceTools/Private/Widgets/STraceControlToolbar.h b/Engine/Source/Developer/TraceTools/Private/Widgets/STraceControlToolbar.h
--- a/Engine/Source/Developer/TraceTools/Private/Widgets/STraceControlToolbar.h
+++ b/Engine/Source/Developer/TraceTools/Private/Widgets/STraceControlToolbar.h
@@ -12,6 +12,8 @@ class ITraceController;
 namespace UE::TraceTools
 {
 
+class SToggleTraceButton;
+
 enum ETraceTarget : uint8
 {
 	Server = 0,
@@ -50,11 +52,8 @@ private:
 	bool SetTraceTarget_CanExecute() const;
 	void SetTraceTarget_Execute(ETraceTarget InTraceTarget);
 
-	bool StartTrace_CanExecute() const;
-	void StartTrace_Execute();
-
-	bool StopTrace_CanExecute() const;
-	void StopTrace_Execute();
+	bool ToggleTrace_CanExecute() const;
+	void ToggleTrace_Execute();
 
 	bool TraceSnapshot_CanExecute() const;
 	void TraceSnapshot_Execute();
diff --git a/Engine/Source/Developer/TraceTools/Public/Interfaces/ITraceToolsModule.h b/Engine/Source/Developer/TraceTools/Public/Interfaces/ITraceToolsModule.h
deleted file mode 100644
diff --git a/Engine/Source/Developer/TraceTools/Public/TraceTools/Interfaces/ITraceToolsModule.h b/Engine/Source/Developer/TraceTools/Public/TraceTools/Interfaces/ITraceToolsModule.h
new file mode 100644
--- /dev/null
+++ b/Engine/Source/Developer/TraceTools/Public/TraceTools/Interfaces/ITraceToolsModule.h
@@ -0,0 +1,41 @@
+// Copyright Epic Games, Inc. All Rights Reserved.
+
+/*=============================================================================
+	ITraceToolsModule.h: Declares the ITraceToolsModule interface.
+=============================================================================*/
+
+#pragma once
+
+#include "Modules/ModuleInterface.h"
+#include "Templates/SharedPointer.h"
+
+class ITraceController;
+class SWidget;
+
+namespace UE::TraceTools
+{
+
+/**
+ * Interface for a trace tools module.
+ */
+class ITraceToolsModule
+	: public IModuleInterface
+{
+public:
+	
+	/**
+	 * Creates a trace control widget.
+	 *
+	 * @return The new widget
+	 */
+	virtual TSharedRef<SWidget> CreateTraceControlWidget(TSharedPtr<ITraceController> InTraceController) = 0;
+
+public:
+
+	/**
+	 * Virtual destructor.
+	 */
+	virtual ~ITraceToolsModule( ) { }
+};
+
+} // namespace UE::TraceTools
diff --git a/Engine/Source/Developer/TraceTools/Public/TraceTools/Widgets/SToggleTraceButton.h b/Engine/Source/Developer/TraceTools/Public/TraceTools/Widgets/SToggleTraceButton.h
new file mode 100644
--- /dev/null
+++ b/Engine/Source/Developer/TraceTools/Public/TraceTools/Widgets/SToggleTraceButton.h
@@ -0,0 +1,76 @@
+// Copyright Epic Games, Inc. All Rights Reserved.
+
+#pragma once
+
+#include "Widgets/DeclarativeSyntaxSupport.h"
+#include "Widgets/SCompoundWidget.h"
+
+class ITraceController;
+
+namespace UE::TraceTools
+{
+
+/**
+ * A dynamic button that can start/stop trace. 
+ * When tracing is started the button turns red and has a pulse animation.
+ */
+class SToggleTraceButton
+	: public SCompoundWidget
+{
+	SLATE_DECLARE_WIDGET_API(SToggleTraceButton, SCompoundWidget, TRACETOOLS_API)
+
+public:
+	enum class EButtonSize : uint32
+	{
+		StatusBar = 0, // 16x16
+		SlimToolbar = 1, // 20x20
+	};
+
+	DECLARE_DELEGATE(FOnToggleTraceRequested);
+
+	SLATE_BEGIN_ARGS(SToggleTraceButton) 
+		: _ButtonSize(EButtonSize::StatusBar)
+	{ }
+	
+	/* An event called when the toggle trace button is clicked. */
+	SLATE_EVENT(FOnToggleTraceRequested, OnToggleTraceRequested)
+	
+	/* Specifies if trace is running at the moment of call. */
+	SLATE_ATTRIBUTE(bool, IsTraceRunning)
+	
+	/* Specifies the size of the button from a few presets. */
+	SLATE_ARGUMENT(EButtonSize, ButtonSize)
+
+	SLATE_END_ARGS()
+
+public:
+	TRACETOOLS_API SToggleTraceButton();
+	TRACETOOLS_API virtual ~SToggleTraceButton();
+
+	TRACETOOLS_API void Construct( const FArguments& InArgs);
+
+private:
+	FSlateColor GetRecordingButtonColor() const;
+	FSlateColor GetRecordingButtonOutlineColor() const;
+	FText GetRecordingButtonTooltipText() const;
+
+	EVisibility GetStartTraceIconVisibility() const;
+	EVisibility GetStopTraceIconVisibility() const;
+
+	void ToggleTrace_OnClicked() const;
+
+	const FSlateBrush* GetToggleTraceCenterBrush() const;
+	const FSlateBrush* GetToggleTraceOutlineBrush() const;
+	const FSlateBrush* GetToggleTraceStopBrush() const;
+
+private:
+	FOnToggleTraceRequested OnToggleTraceRequested;
+
+	bool bIsTraceRecordButtonHovered = false;
+	mutable double ConnectionStartTime = 0.0f;
+
+	TSlateAttribute<bool> IsTraceRunningAttribute;
+	EButtonSize ButtonSize;
+};
+
+} // namespace UE::TraceTools
\ No newline at end of file
