Add the trace endpoint to the trace control window.
Make a single dynamic Pause/Resume dynamic button.

diff --git a/Engine/Source/Developer/TraceTools/Private/Models/TraceControlCommands.h b/Engine/Source/Developer/TraceTools/Private/Models/TraceControlCommands.h
--- a/Engine/Source/Developer/TraceTools/Private/Models/TraceControlCommands.h
+++ b/Engine/Source/Developer/TraceTools/Private/Models/TraceControlCommands.h
@@ -37,9 +37,6 @@ public:
 		UI_COMMAND(StopTrace, "Stop Trace", "Stop trace on the selected session.", EUserInterfaceActionType::Button, FInputChord());
 		UI_COMMAND(TraceSnapshot, "Snapshot", "Save a trace snapshot on the selected session.", EUserInterfaceActionType::Button, FInputChord());
 
-		UI_COMMAND(PauseTrace, "Pause Trace", "Disable all the trace channels and save the channel list so it can be enabled again with the resume command.", EUserInterfaceActionType::Button, FInputChord());
-		UI_COMMAND(ResumeTrace, "Resume Trace", "Enable the channels that were enabled before trace was paused.", EUserInterfaceActionType::Button, FInputChord());
-
 		UI_COMMAND(TraceBookmark, "Bookmark", "Trace a bookmark with a predefined name.", EUserInterfaceActionType::Button, FInputChord());
 		UI_COMMAND(TraceScreenshot, "Screenshot", "Trace a screenshot with a predefined name.", EUserInterfaceActionType::Button, FInputChord());
 
@@ -52,8 +49,6 @@ public:
 	TSharedPtr<FUICommandInfo> StartTrace;
 	TSharedPtr<FUICommandInfo> StopTrace;
 	TSharedPtr<FUICommandInfo> TraceSnapshot;
-	TSharedPtr<FUICommandInfo> PauseTrace;
-	TSharedPtr<FUICommandInfo> ResumeTrace;
 	TSharedPtr<FUICommandInfo> TraceBookmark;
 	TSharedPtr<FUICommandInfo> TraceScreenshot;
 	TSharedPtr<FUICommandInfo> ToggleStatNamedEvents;
diff --git a/Engine/Source/Developer/TraceTools/Private/Services/ISessionTraceFilterService.h b/Engine/Source/Developer/TraceTools/Private/Services/ISessionTraceFilterService.h
--- a/Engine/Source/Developer/TraceTools/Private/Services/ISessionTraceFilterService.h
+++ b/Engine/Source/Developer/TraceTools/Private/Services/ISessionTraceFilterService.h
@@ -43,16 +43,16 @@ class ISessionTraceFilterService : public TSharedFromThis<ISessionTraceFilterSer
 public:
 	virtual ~ISessionTraceFilterService() {}
 
-	/** Returns the root level set of objects */
+	/** Returns the root level set of objects. */
 	virtual void GetRootObjects(TArray<FTraceObjectInfo>& OutObjects) const = 0;
 	
-	/** Returns the object with the specified name */
+	/** Returns the object with the specified name. */
 	virtual const FTraceObjectInfo* GetObject(const FString& Name ) const = 0;
 
-	/** Set the filtered state for an individual object by its hash */
+	/** Set the filtered state for an individual object by its hash. */
 	virtual void SetObjectFilterState(const FString& InObjectName, const bool bFilterState) = 0;
 	
-	/** Set timestamp for last processed update (data change) */
+	/** Set timestamp for last processed update (data change). */
 	virtual const FDateTime& GetTimestamp() const = 0;
 
 	/** Update filter preset */
@@ -61,14 +61,17 @@ public:
 	/** Returns true if settings are available for the selected session. */
 	virtual bool HasSettings() const = 0 ;
 
-	/** Get the settings of the selected session */
+	/** Get the settings of the selected session. */
 	virtual const FTraceStatus::FSettings& GetSettings() const = 0;
 
 	/** Returns true if stats are available for the selected session. */
 	virtual bool HasStats() const = 0;
 
-	/** Get the stats of the selected session */
+	/** Get the stats of the selected session. */
 	virtual const FTraceStats& GetStats() const = 0;
+
+	/** Get the endpoint of the current running trace. */
+	virtual const FString& GetTraceEndpoint() const = 0;
 };
 
 } // namespace UE::TraceTools
\ No newline at end of file
diff --git a/Engine/Source/Developer/TraceTools/Private/Services/SessionTraceControllerFilterService.cpp b/Engine/Source/Developer/TraceTools/Private/Services/SessionTraceControllerFilterService.cpp
--- a/Engine/Source/Developer/TraceTools/Private/Services/SessionTraceControllerFilterService.cpp
+++ b/Engine/Source/Developer/TraceTools/Private/Services/SessionTraceControllerFilterService.cpp
@@ -125,6 +125,8 @@ void FSessionTraceControllerFilterService::OnTraceStatusUpdated(const FTraceStat
 
 	if (EnumHasAnyFlags(InUpdateType, FTraceStatus::EUpdateType::Status))
 	{
+		TraceEndpoint = InStatus.Endpoint;
+
 		Stats.BytesSentPerSecond = 0;
 		Stats.BytesTracedPerSecond = 0;
 
@@ -208,6 +210,7 @@ void FSessionTraceControllerFilterService::OnSessionSelectionChanged()
 {
 	bHasStats = false;
 	bHasSettings = false;
+	TraceEndpoint.Empty();
 }
 
 } // namespace UE::TraceTools
\ No newline at end of file
diff --git a/Engine/Source/Developer/TraceTools/Private/Services/SessionTraceControllerFilterService.h b/Engine/Source/Developer/TraceTools/Private/Services/SessionTraceControllerFilterService.h
--- a/Engine/Source/Developer/TraceTools/Private/Services/SessionTraceControllerFilterService.h
+++ b/Engine/Source/Developer/TraceTools/Private/Services/SessionTraceControllerFilterService.h
@@ -36,6 +36,8 @@ public:
 
 	virtual bool HasStats() const override;
 	virtual const FTraceStats& GetStats() const override;
+
+	virtual const FString& GetTraceEndpoint() const override { return TraceEndpoint; }
 	/** End ISessionTraceFilterService overrides */
 
 protected:
@@ -71,6 +73,8 @@ protected:
 
 	FTraceStats Stats;
 	bool bHasStats = false;
+
+	FString TraceEndpoint;
 };
 
 } // namespace UE::TraceTools
\ No newline at end of file
diff --git a/Engine/Source/Developer/TraceTools/Private/TraceToolsStyle.cpp b/Engine/Source/Developer/TraceTools/Private/TraceToolsStyle.cpp
--- a/Engine/Source/Developer/TraceTools/Private/TraceToolsStyle.cpp
+++ b/Engine/Source/Developer/TraceTools/Private/TraceToolsStyle.cpp
@@ -202,6 +202,8 @@ void FTraceToolsStyle::Initialize()
 	StyleSet->Set("ToggleTraceButton.RecordTraceRecording.SlimToolbar", new CORE_IMAGE_BRUSH_SVG("Starship/Insights/TraceTools/RecordTraceRecording", Icon20x20));
 	StyleSet->Set("ToggleTraceButton.TraceStop.SlimToolbar", new CORE_IMAGE_BRUSH_SVG("Starship/Insights/TraceTools/TraceStop", Icon20x20, FStyleColors::Error));
 	
+	StyleSet->Set("TraceStatistics.CopyEndpoint", new CORE_IMAGE_BRUSH_SVG("Starship/Common/Copy", FVector2D(10.0f, 10.0f)));
+	
 	FSlateStyleRegistry::RegisterSlateStyle( *StyleSet.Get() );
 }
 
diff --git a/Engine/Source/Developer/TraceTools/Private/Widgets/STraceControlToolbar.cpp b/Engine/Source/Developer/TraceTools/Private/Widgets/STraceControlToolbar.cpp
--- a/Engine/Source/Developer/TraceTools/Private/Widgets/STraceControlToolbar.cpp
+++ b/Engine/Source/Developer/TraceTools/Private/Widgets/STraceControlToolbar.cpp
@@ -8,6 +8,9 @@
 #include "SlateOptMacros.h"
 #include "SocketSubsystem.h"
 #include "Styling/AppStyle.h"
+#include "Styling/StyleColors.h"
+#include "Widgets/Images/SImage.h"
+#include "Widgets/Input/SButton.h"
 #include "Widgets/Layout/SBorder.h"
 #include "Widgets/SBoxPanel.h"
 
@@ -77,8 +80,28 @@ void STraceControlToolbar::Construct(const FArguments& InArgs, const TSharedRef<
 
 		Toolbar.AddSeparator();
 
-		Toolbar.AddToolBarButton(FTraceControlCommands::Get().PauseTrace);
-		Toolbar.AddToolBarButton(FTraceControlCommands::Get().ResumeTrace);
+		TSharedRef<SWidget> PauseResumeWidget =
+			SNew(SHorizontalBox)
+
+			+ SHorizontalBox::Slot()
+			[
+				SNew(SButton)
+				.ButtonStyle(FAppStyle::Get(), "SimpleButton")
+				.ContentPadding(FMargin(0.0f, 0.0f, 0.0f, 3.0f))
+				.HAlign(HAlign_Left)
+				.VAlign(VAlign_Bottom)
+				.ToolTipText(this, &STraceControlToolbar::TogglePauseResume_GetTooltip)
+				.OnClicked(this, &STraceControlToolbar::TogglePauseResume_OnClicked)
+				.IsEnabled(this, &STraceControlToolbar::TogglePauseResume_CanExecute)
+				.Content()
+				[
+					SNew(SImage)
+					.Image(this, &STraceControlToolbar::GetPauseResumeBrush)
+					.ColorAndOpacity(FStyleColors::Foreground)
+				]
+			];
+
+		Toolbar.AddWidget(PauseResumeWidget);
 
 		Toolbar.AddSeparator();
 
@@ -128,14 +151,6 @@ void STraceControlToolbar::BindCommands(const TSharedRef<FUICommandList>& Comman
 						   FExecuteAction::CreateSP(this, &STraceControlToolbar::TraceSnapshot_Execute),
 						   FCanExecuteAction::CreateSP(this, &STraceControlToolbar::TraceSnapshot_CanExecute));
 
-	CommandList->MapAction(FTraceControlCommands::Get().PauseTrace, 
-						   FExecuteAction::CreateSP(this, &STraceControlToolbar::PauseTrace_Execute),
-						   FCanExecuteAction::CreateSP(this, &STraceControlToolbar::PauseTrace_CanExecute));
-
-	CommandList->MapAction(FTraceControlCommands::Get().ResumeTrace, 
-						   FExecuteAction::CreateSP(this, &STraceControlToolbar::ResumeTrace_Execute),
-						   FCanExecuteAction::CreateSP(this, &STraceControlToolbar::ResumeTrace_CanExecute));
-
 	CommandList->MapAction(FTraceControlCommands::Get().TraceBookmark, 
 						   FExecuteAction::CreateSP(this, &STraceControlToolbar::TraceBookmark_Execute),
 						   FCanExecuteAction::CreateSP(this, &STraceControlToolbar::TraceBookmark_CanExecute));
@@ -231,34 +246,6 @@ void STraceControlToolbar::TraceSnapshot_Execute()
 	});
 }
 
-bool STraceControlToolbar::PauseTrace_CanExecute() const
-{
-	return TraceController->HasAvailableSelectedInstance() && bIsTracing && bIsPaused == false;
-}
-
-void STraceControlToolbar::PauseTrace_Execute()
-{
-	TraceController->WithSelectedInstances([&](const FTraceStatus& Status, ITraceControllerCommands& Commands)
-	{
-		Commands.Pause();
-	});
-	bIsPaused = true;
-}
-
-bool STraceControlToolbar::ResumeTrace_CanExecute() const
-{
-	return TraceController->HasAvailableSelectedInstance() && bIsTracing && bIsPaused == true;
-}
-
-void STraceControlToolbar::ResumeTrace_Execute()
-{
-	TraceController->WithSelectedInstances([&](const FTraceStatus& Status, ITraceControllerCommands& Commands)
-	{
-		Commands.Resume();
-	});
-	bIsPaused = false;
-}
-
 bool STraceControlToolbar::TraceBookmark_CanExecute() const
 {
 	return TraceController->HasAvailableSelectedInstance() && bIsTracing && !bIsPaused;
@@ -337,6 +324,53 @@ void STraceControlToolbar::OnTraceStatusUpdated(const FTraceStatus& InStatus, FT
 	bAreStatNamedEventsEnabled = InStatus.bAreStatNamedEventsEnabled;
 }
 
+bool STraceControlToolbar::TogglePauseResume_CanExecute() const
+{
+	return TraceController->HasAvailableSelectedInstance() && bIsTracing;
+}
+
+FReply STraceControlToolbar::TogglePauseResume_OnClicked()
+{
+	if (bIsPaused)
+	{
+		TraceController->WithSelectedInstances([&](const FTraceStatus& Status, ITraceControllerCommands& Commands)
+			{
+				Commands.Resume();
+			});
+		bIsPaused = false;
+	}
+	else
+	{
+		TraceController->WithSelectedInstances([&](const FTraceStatus& Status, ITraceControllerCommands& Commands)
+			{
+				Commands.Pause();
+			});
+		bIsPaused = true;
+	}
+
+	return FReply::Handled();
+}
+
+const FSlateBrush* STraceControlToolbar::GetPauseResumeBrush() const
+{
+	if (bIsPaused)
+	{
+		return FTraceToolsStyle::GetBrush("TraceControl.ResumeTrace.Small");
+	}
+		
+	return FTraceToolsStyle::GetBrush("TraceControl.PauseTrace.Small");
+}
+
+FText STraceControlToolbar::TogglePauseResume_GetTooltip() const
+{
+	if (bIsPaused)
+	{
+		return LOCTEXT("ResumeTraceTooltip", "Enable the channels that were enabled before trace was paused.");
+	}
+
+	return LOCTEXT("PauseTraceTooltip", "Disable all the trace channels and save the channel list so they can be enabled again with the resume command.");
+}
+
 } // namespace UE::TraceTools
 
 #undef LOCTEXT_NAMESPACE
\ No newline at end of file
diff --git a/Engine/Source/Developer/TraceTools/Private/Widgets/STraceControlToolbar.h b/Engine/Source/Developer/TraceTools/Private/Widgets/STraceControlToolbar.h
--- a/Engine/Source/Developer/TraceTools/Private/Widgets/STraceControlToolbar.h
+++ b/Engine/Source/Developer/TraceTools/Private/Widgets/STraceControlToolbar.h
@@ -8,6 +8,8 @@
 #include "Widgets/SCompoundWidget.h"
 
 class ITraceController;
+class FReply;
+struct FSlateBrush;
 
 namespace UE::TraceTools
 {
@@ -58,12 +60,6 @@ private:
 	bool TraceSnapshot_CanExecute() const;
 	void TraceSnapshot_Execute();
 
-	bool PauseTrace_CanExecute() const;
-	void PauseTrace_Execute();
-
-	bool ResumeTrace_CanExecute() const;
-	void ResumeTrace_Execute();
-
 	bool TraceBookmark_CanExecute() const;
 	void TraceBookmark_Execute();
 
@@ -74,6 +70,12 @@ private:
 	bool ToggleStatNamedEvents_IsChecked() const;
 	void ToggleStatNamedEvents_Execute();
 
+	FReply TogglePauseResume_OnClicked();
+	bool TogglePauseResume_CanExecute() const;
+	FText TogglePauseResume_GetTooltip() const;
+
+	const FSlateBrush* GetPauseResumeBrush() const;
+
 private:
 	TSharedPtr<ITraceController> TraceController;
 
diff --git a/Engine/Source/Developer/TraceTools/Private/Widgets/STraceStatistics.cpp b/Engine/Source/Developer/TraceTools/Private/Widgets/STraceStatistics.cpp
--- a/Engine/Source/Developer/TraceTools/Private/Widgets/STraceStatistics.cpp
+++ b/Engine/Source/Developer/TraceTools/Private/Widgets/STraceStatistics.cpp
@@ -2,9 +2,12 @@
 
 #include "STraceStatistics.h"
 
+#include "HAL/PlatformApplicationMisc.h"
 #include "Internationalization/Text.h"
 #include "SlateOptMacros.h"
 #include "Styling/StyleColors.h"
+#include "Widgets/Images/SImage.h"
+#include "Widgets/Input/SButton.h"
 #include "Widgets/Layout/SBorder.h"
 #include "Widgets/SBoxPanel.h"
 #include "Widgets/Text/STextBlock.h"
@@ -42,230 +45,283 @@ void STraceStatistics::Construct(const FArguments& InArgs, TSharedPtr<ISessionTr
 			SNew(SBorder)
 			.BorderImage(FTraceToolsStyle::GetBrush("FilterPresets.BackgroundBorder"))
 			[
-				SNew(SHorizontalBox)
-
-				+ SHorizontalBox::Slot()
-				.AutoWidth()
+				SNew(SVerticalBox)
+				
+				+ SVerticalBox::Slot()
+				.HAlign(EHorizontalAlignment::HAlign_Left)
+				.Padding(0.0f, 3.0f, 0.0f, 0.0f)
+				.AutoHeight()
 				[
-					SNew(SVerticalBox)
-
-					// Trace Settings
-					+ SVerticalBox::Slot()
-					.HAlign(EHorizontalAlignment::HAlign_Left)
-					.Padding(0.0f, 10.0f, 0.0f, 0.0f)
-					.AutoHeight()
+					SNew(SHorizontalBox)
+							
+					+ SHorizontalBox::Slot()
+					.Padding(0.0f, 2.0f, 0.0f, 0.0f)
+					.AutoWidth()
 					[
 						SNew(STextBlock)
-						.Text(LOCTEXT("Trace Settings", "Trace Settings"))
-						.Font(FAppStyle::Get().GetFontStyle("NormalFontBold"))
+						.ColorAndOpacity(FSlateColor(EStyleColor::Foreground))
+						.ToolTipText(LOCTEXT("TraceEndpointTooltip", "The endpoint the current trace is sending data to."))
+						.Text(LOCTEXT("TraceEndpoint", "Trace Endpoint:"))
 					]
 
-					+ SVerticalBox::Slot()
-					.HAlign(EHorizontalAlignment::HAlign_Left)
-					.AutoHeight()
+					+ SHorizontalBox::Slot()
+					.Padding(2.0f, 2.0f, 0.0f, 0.0f)
+					.AutoWidth()
 					[
-						SNew(SHorizontalBox)
-							
-						+ SHorizontalBox::Slot()
-						.Padding(0.0f, 2.0f, 0.0f, 0.0f)
-						.AutoWidth()
-						[
-							SNew(STextBlock)
-							.ColorAndOpacity(FSlateColor(EStyleColor::Foreground))
-							.ToolTipText(LOCTEXT("ImportantEventsSettingTooltip", "The state of the Important Events cache."))
-							.Text(LOCTEXT("ImportantCache", "Important Events Cache:"))
-						]
-
-						+ SHorizontalBox::Slot()
-						.Padding(2.0f, 2.0f, 0.0f, 0.0f)
-						.AutoWidth()
-						[
-							SNew(STextBlock)
-							.ColorAndOpacity(FSlateColor(EStyleColor::AccentGray))
-							.Text_Lambda([this]() { return this->GetSettingsOnOffText(SessionFilterService->GetSettings().bUseImportantCache); })
-						]
+						SNew(STextBlock)
+						.ColorAndOpacity(FSlateColor(EStyleColor::AccentGray))
+						.Text(this, &STraceStatistics::GetTraceEndpointText)
 					]
 
-					+ SVerticalBox::Slot()
-					.HAlign(EHorizontalAlignment::HAlign_Left)
-					.AutoHeight()
+					+ SHorizontalBox::Slot()
+					.Padding(2.0f, 2.0f, 0.0f, 0.0f)
+					.AutoWidth()
 					[
-						SNew(SHorizontalBox)
-							
-						+ SHorizontalBox::Slot()
-						.Padding(0.0f, 2.0f, 0.0f, 0.0f)
-						.AutoWidth()
-						[
-							SNew(STextBlock)
-							.ColorAndOpacity(FSlateColor(EStyleColor::Foreground))
-							.ToolTipText(LOCTEXT("UseWorkerThreadTooltip", "If trace uses a worker thread. If not, TraceLog is pumped on end frame."))
-							.Text(LOCTEXT("WorkerThread", "Worker Thread:"))
-						]
-
-						+ SHorizontalBox::Slot()
-						.Padding(2.0f, 2.0f, 0.0f, 0.0f)
-						.AutoWidth()
+						SNew(SButton)
+						.ButtonStyle(FAppStyle::Get(), "SimpleButton")
+						.ContentPadding(FMargin(0.0f, 0.0f, 0.0f, 0.0f))
+						.HAlign(HAlign_Left)
+						.VAlign(VAlign_Bottom)
+						.ToolTipText(LOCTEXT("CopyEndpointTooltip", "Copy the value of the current endpoint."))
+						.OnClicked(this, &STraceStatistics::CopyEndpoint_OnClicked)
+						.Content()
 						[
-							SNew(STextBlock)
-							.ColorAndOpacity(FSlateColor(EStyleColor::AccentGray))
-							.Text_Lambda([this]() { return this->GetSettingsOnOffText(SessionFilterService->GetSettings().bUseWorkerThread); })
+							SNew(SImage)
+							.Image(FTraceToolsStyle::GetBrush("TraceStatistics.CopyEndpoint"))
+							.Visibility(this, &STraceStatistics::GetCopyEndpointVisibility)
 						]
 					]
+				]
+				
+				+ SVerticalBox::Slot()
+				.Padding(0.0f, 5.0f, 0.0f, 0.0f)
+				[
+					SNew(SHorizontalBox)
 
-					+ SVerticalBox::Slot()
-					.HAlign(EHorizontalAlignment::HAlign_Left)
-					.AutoHeight()
+					+ SHorizontalBox::Slot()
+					.AutoWidth()
 					[
-						SNew(SHorizontalBox)
-							
-						+ SHorizontalBox::Slot()
-						.Padding(0.0f, 2.0f, 0.0f, 0.0f)
-						.AutoWidth()
-						[
-							SNew(STextBlock)
-							.ColorAndOpacity(FSlateColor(EStyleColor::Foreground))
-							.ToolTipText(LOCTEXT("TailSizeTooltip", "Size of the tail buffer where the last seconds of trace data are stored."))
-							.Text(LOCTEXT("TailSize", "Tail Size:"))
-						]
+						SNew(SVerticalBox)
 
-						+ SHorizontalBox::Slot()
-						.Padding(2.0f, 2.0f, 0.0f, 0.0f)
-						.AutoWidth()
+						// Trace Settings
+						+ SVerticalBox::Slot()
+						.HAlign(EHorizontalAlignment::HAlign_Left)
+						.Padding(0.0f, 10.0f, 0.0f, 0.0f)
+						.AutoHeight()
 						[
 							SNew(STextBlock)
-							.ColorAndOpacity(FSlateColor(EStyleColor::AccentGray))
-							.Text_Lambda([this]() { return this->GetSettingsMemoryValueText(SessionFilterService->GetSettings().TailSizeBytes); })
+							.Text(LOCTEXT("Trace Settings", "Trace Settings"))
+							.Font(FAppStyle::Get().GetFontStyle("NormalFontBold"))
 						]
-					]
-				]
-
-				+ SHorizontalBox::Slot()
-				.AutoWidth()
-				.Padding(30.0f, 0.0f, 0.0f, 0.0f)
-				[
-					SNew(SVerticalBox)
 
-					// Trace statistics
-					+ SVerticalBox::Slot()
-					.HAlign(EHorizontalAlignment::HAlign_Left)
-					.Padding(0.0f, 10.0f, 0.0f, 0.0f)
-					.AutoHeight()
-					[
-						SNew(STextBlock)
-						.Text(LOCTEXT("Statistics", "Statistics"))
-						.Font(FAppStyle::Get().GetFontStyle("NormalFontBold"))
-					]
+						+ SVerticalBox::Slot()
+						.HAlign(EHorizontalAlignment::HAlign_Left)
+						.AutoHeight()
+						[
+							SNew(SHorizontalBox)
 
-					+ SVerticalBox::Slot()
-					[
-						SNew(SHorizontalBox)
 							
-						+ SHorizontalBox::Slot()
-						.Padding(0.0f, 2.0, 0.0f, 0.0f)
-						.AutoWidth()
-						[
-							SNew(STextBlock)
-							.ColorAndOpacity(FSlateColor(EStyleColor::Foreground))
-							.ToolTipText(LOCTEXT("BytesSentTooltip", "Number of bytes sent to server or file."))
-							.Text(LOCTEXT("BytesSent", "Bytes Sent:"))
+							+ SHorizontalBox::Slot()
+							.Padding(0.0f, 2.0f, 0.0f, 0.0f)
+							.AutoWidth()
+							[
+								SNew(STextBlock)
+								.ColorAndOpacity(FSlateColor(EStyleColor::Foreground))
+								.ToolTipText(LOCTEXT("ImportantEventsSettingTooltip", "The state of the Important Events cache."))
+								.Text(LOCTEXT("ImportantCache", "Important Events Cache:"))
+							]
+
+							+ SHorizontalBox::Slot()
+							.Padding(2.0f, 2.0f, 0.0f, 0.0f)
+							.AutoWidth()
+							[
+								SNew(STextBlock)
+								.ColorAndOpacity(FSlateColor(EStyleColor::AccentGray))
+								.Text_Lambda([this]() { return this->GetSettingsOnOffText(SessionFilterService->GetSettings().bUseImportantCache); })
+							]
 						]
 
-						+ SHorizontalBox::Slot()
-						.AutoWidth()
-						.Padding(2.0f, 2.0, 0.0f, 0.0f)
+						+ SVerticalBox::Slot()
+						.HAlign(EHorizontalAlignment::HAlign_Left)
+						.AutoHeight()
 						[
-							SNew(STextBlock)
-							.ColorAndOpacity(FSlateColor(EStyleColor::AccentGray))
-							.Text_Lambda([this]() { return this->GetStatsMemoryValueText(SessionFilterService->GetStats().StandardStats.BytesSent); })
+							SNew(SHorizontalBox)
+							
+							+ SHorizontalBox::Slot()
+							.Padding(0.0f, 2.0f, 0.0f, 0.0f)
+							.AutoWidth()
+							[
+								SNew(STextBlock)
+								.ColorAndOpacity(FSlateColor(EStyleColor::Foreground))
+								.ToolTipText(LOCTEXT("UseWorkerThreadTooltip", "If trace uses a worker thread. If not, TraceLog is pumped on end frame."))
+								.Text(LOCTEXT("WorkerThread", "Worker Thread:"))
+							]
+
+							+ SHorizontalBox::Slot()
+							.Padding(2.0f, 2.0f, 0.0f, 0.0f)
+							.AutoWidth()
+							[
+								SNew(STextBlock)
+								.ColorAndOpacity(FSlateColor(EStyleColor::AccentGray))
+								.Text_Lambda([this]() { return this->GetSettingsOnOffText(SessionFilterService->GetSettings().bUseWorkerThread); })
+							]
 						]
 
-						+ SHorizontalBox::Slot()
-						.AutoWidth()
-						.Padding(4.0f, 2.0, 0.0f, 0.0f)
+						+ SVerticalBox::Slot()
+						.HAlign(EHorizontalAlignment::HAlign_Left)
+						.AutoHeight()
 						[
-							SNew(STextBlock)
-							.ColorAndOpacity(FSlateColor(EStyleColor::AccentGray))
-							.Text_Lambda([this]() { return this->GetStatsBandwidthText(SessionFilterService->GetStats().BytesSentPerSecond); })
+							SNew(SHorizontalBox)
+							
+							+ SHorizontalBox::Slot()
+							.Padding(0.0f, 2.0f, 0.0f, 0.0f)
+							.AutoWidth()
+							[
+								SNew(STextBlock)
+								.ColorAndOpacity(FSlateColor(EStyleColor::Foreground))
+								.ToolTipText(LOCTEXT("TailSizeTooltip", "Size of the tail buffer where the last seconds of trace data are stored."))
+								.Text(LOCTEXT("TailSize", "Tail Size:"))
+							]
+
+							+ SHorizontalBox::Slot()
+							.Padding(2.0f, 2.0f, 0.0f, 0.0f)
+							.AutoWidth()
+							[
+								SNew(STextBlock)
+								.ColorAndOpacity(FSlateColor(EStyleColor::AccentGray))
+								.Text_Lambda([this]() { return this->GetSettingsMemoryValueText(SessionFilterService->GetSettings().TailSizeBytes); })
+							]
 						]
 					]
 
-					+ SVerticalBox::Slot()
+					+ SHorizontalBox::Slot()
+					.AutoWidth()
+					.Padding(30.0f, 0.0f, 0.0f, 0.0f)
 					[
-						SNew(SHorizontalBox)
-							
-						+ SHorizontalBox::Slot()
-						.Padding(0.0f, 2.0, 0.0f, 0.0f)
-						.AutoWidth()
-						[
-							SNew(STextBlock)
-							.ColorAndOpacity(FSlateColor(EStyleColor::Foreground))
-							.ToolTipText(LOCTEXT("BytesTracedTooltip", "Number of (uncompressed) bytes traced from process."))
-							.Text(LOCTEXT("BytesTraced", "Bytes Traced:"))
-						]
+						SNew(SVerticalBox)
 
-						+ SHorizontalBox::Slot()
-						.AutoWidth()
-						.Padding(2.0f, 2.0, 0.0f, 0.0f)
+						// Trace statistics
+						+ SVerticalBox::Slot()
+						.HAlign(EHorizontalAlignment::HAlign_Left)
+						.Padding(0.0f, 10.0f, 0.0f, 0.0f)
+						.AutoHeight()
 						[
 							SNew(STextBlock)
-							.ColorAndOpacity(FSlateColor(EStyleColor::AccentGray))
-							.Text_Lambda([this]() { return this->GetStatsMemoryValueText(SessionFilterService->GetStats().StandardStats.BytesTraced); })
+							.Text(LOCTEXT("Statistics", "Statistics"))
+							.Font(FAppStyle::Get().GetFontStyle("NormalFontBold"))
 						]
 
-						+ SHorizontalBox::Slot()
-						.AutoWidth()
-						.Padding(2.0f, 2.0, 0.0f, 0.0f)
+						+ SVerticalBox::Slot()
 						[
-							SNew(STextBlock)
-							.ColorAndOpacity(FSlateColor(EStyleColor::AccentGray))
-							.Text_Lambda([this]() { return this->GetStatsBandwidthText(SessionFilterService->GetStats().BytesTracedPerSecond); })
-						]
-					]
-		
-					+ SVerticalBox::Slot()
-					[
-						SNew(SHorizontalBox)
+							SNew(SHorizontalBox)
 							
-						+ SHorizontalBox::Slot()
-						.Padding(0.0f, 2.0, 0.0f, 0.0f)
-						.AutoWidth()
-						[
-							SNew(STextBlock)
-							.ColorAndOpacity(FSlateColor(EStyleColor::Foreground))
-							.ToolTipText(LOCTEXT("MemoryUsedTooltip", "Total memory used by TraceLog."))
-							.Text(LOCTEXT("MemoryUsed", "Memory Used:"))
+							+ SHorizontalBox::Slot()
+							.Padding(0.0f, 2.0, 0.0f, 0.0f)
+							.AutoWidth()
+							[
+								SNew(STextBlock)
+								.ColorAndOpacity(FSlateColor(EStyleColor::Foreground))
+								.ToolTipText(LOCTEXT("BytesSentTooltip", "Number of bytes sent to server or file."))
+								.Text(LOCTEXT("BytesSent", "Bytes Sent:"))
+							]
+
+							+ SHorizontalBox::Slot()
+							.AutoWidth()
+							.Padding(2.0f, 2.0, 0.0f, 0.0f)
+							[
+								SNew(STextBlock)
+								.ColorAndOpacity(FSlateColor(EStyleColor::AccentGray))
+								.Text_Lambda([this]() { return this->GetStatsMemoryValueText(SessionFilterService->GetStats().StandardStats.BytesSent); })
+							]
+
+							+ SHorizontalBox::Slot()
+							.AutoWidth()
+							.Padding(4.0f, 2.0, 0.0f, 0.0f)
+							[
+								SNew(STextBlock)
+								.ColorAndOpacity(FSlateColor(EStyleColor::AccentGray))
+								.Text_Lambda([this]() { return this->GetStatsBandwidthText(SessionFilterService->GetStats().BytesSentPerSecond); })
+							]
 						]
 
-						+ SHorizontalBox::Slot()
-						.Padding(2.0f, 2.0, 0.0f, 0.0f)
+						+ SVerticalBox::Slot()
 						[
-							SNew(STextBlock)
-							.ColorAndOpacity(FSlateColor(EStyleColor::AccentGray))
-							.Text_Lambda([this]() { return this->GetStatsMemoryValueText(SessionFilterService->GetStats().StandardStats.MemoryUsed); })
+							SNew(SHorizontalBox)
+							
+							+ SHorizontalBox::Slot()
+							.Padding(0.0f, 2.0, 0.0f, 0.0f)
+							.AutoWidth()
+							[
+								SNew(STextBlock)
+								.ColorAndOpacity(FSlateColor(EStyleColor::Foreground))
+								.ToolTipText(LOCTEXT("BytesTracedTooltip", "Number of (uncompressed) bytes traced from process."))
+								.Text(LOCTEXT("BytesTraced", "Bytes Traced:"))
+							]
+
+							+ SHorizontalBox::Slot()
+							.AutoWidth()
+							.Padding(2.0f, 2.0, 0.0f, 0.0f)
+							[
+								SNew(STextBlock)
+								.ColorAndOpacity(FSlateColor(EStyleColor::AccentGray))
+								.Text_Lambda([this]() { return this->GetStatsMemoryValueText(SessionFilterService->GetStats().StandardStats.BytesTraced); })
+							]
+
+							+ SHorizontalBox::Slot()
+							.AutoWidth()
+							.Padding(2.0f, 2.0, 0.0f, 0.0f)
+							[
+								SNew(STextBlock)
+								.ColorAndOpacity(FSlateColor(EStyleColor::AccentGray))
+								.Text_Lambda([this]() { return this->GetStatsBandwidthText(SessionFilterService->GetStats().BytesTracedPerSecond); })
+							]
 						]
-					]
-
-					+ SVerticalBox::Slot()
-					[
-						SNew(SHorizontalBox)
-
-						+ SHorizontalBox::Slot()
-						.Padding(0.0f, 2.0, 0.0f, 0.0f)
-						.AutoWidth()
+		
+						+ SVerticalBox::Slot()
 						[
-							SNew(STextBlock)
-							.ColorAndOpacity(FSlateColor(EStyleColor::Foreground))
-							.ToolTipText(LOCTEXT("ImportantEventsMemoryTooltip", "Memory for important events."))
-							.Text(LOCTEXT("ImportantEventsCache:", "Cache:"))
+							SNew(SHorizontalBox)
+							
+							+ SHorizontalBox::Slot()
+							.Padding(0.0f, 2.0, 0.0f, 0.0f)
+							.AutoWidth()
+							[
+								SNew(STextBlock)
+								.ColorAndOpacity(FSlateColor(EStyleColor::Foreground))
+								.ToolTipText(LOCTEXT("MemoryUsedTooltip", "Total memory used by TraceLog."))
+								.Text(LOCTEXT("MemoryUsed", "Memory Used:"))
+							]
+
+							+ SHorizontalBox::Slot()
+							.Padding(2.0f, 2.0, 0.0f, 0.0f)
+							[
+								SNew(STextBlock)
+								.ColorAndOpacity(FSlateColor(EStyleColor::AccentGray))
+								.Text_Lambda([this]() { return this->GetStatsMemoryValueText(SessionFilterService->GetStats().StandardStats.MemoryUsed); })
+							]
 						]
 
-						+ SHorizontalBox::Slot()
-						.AutoWidth()
-						.Padding(2.0f, 2.0, 0.0f, 0.0f)
+						+ SVerticalBox::Slot()
 						[
-							SNew(STextBlock)
-							.ColorAndOpacity(FSlateColor(EStyleColor::AccentGray))
-							.Text_Lambda([this]() { return this->GetStatsCacheText(); })
+							SNew(SHorizontalBox)
+
+							+ SHorizontalBox::Slot()
+							.Padding(0.0f, 2.0, 0.0f, 0.0f)
+							.AutoWidth()
+							[
+								SNew(STextBlock)
+								.ColorAndOpacity(FSlateColor(EStyleColor::Foreground))
+								.ToolTipText(LOCTEXT("ImportantEventsMemoryTooltip", "Memory for important events."))
+								.Text(LOCTEXT("ImportantEventsCache:", "Cache:"))
+							]
+
+							+ SHorizontalBox::Slot()
+							.AutoWidth()
+							.Padding(2.0f, 2.0, 0.0f, 0.0f)
+							[
+								SNew(STextBlock)
+								.ColorAndOpacity(FSlateColor(EStyleColor::AccentGray))
+								.Text_Lambda([this]() { return this->GetStatsCacheText(); })
+							]
 						]
 					]
 				]
@@ -358,6 +414,38 @@ FText STraceStatistics::GetStatsCacheText() const
 	return FText::Format(LOCTEXT("TraceCacheTextFormat", "{0} ({1} used + {2} unused | {3} waste)"), CacheAllocated, CacheUsed, CacheUnused, CacheWasted);
 }
 
+FText STraceStatistics::GetTraceEndpointText() const
+{
+	if (!SessionFilterService->HasStats() || SessionFilterService->GetTraceEndpoint().IsEmpty())
+	{
+		return LOCTEXT("N/A", "N/A");
+	}
+
+	return FText::FromString(SessionFilterService->GetTraceEndpoint());
+}
+
+FReply STraceStatistics::CopyEndpoint_OnClicked() const
+{
+	const FString& Endpoint = SessionFilterService->GetTraceEndpoint();
+
+	if (!Endpoint.IsEmpty())
+	{
+		FPlatformApplicationMisc::ClipboardCopy(*Endpoint);
+	}
+
+	return FReply::Handled();
+}
+
+EVisibility STraceStatistics::GetCopyEndpointVisibility() const
+{
+	if (!SessionFilterService->HasStats() || SessionFilterService->GetTraceEndpoint().IsEmpty())
+	{
+		return EVisibility::Collapsed;
+	}
+
+	return EVisibility::Visible;
+}
+
 } // namespace UE::TraceTools
 
 #undef LOCTEXT_NAMESPACE
\ No newline at end of file
diff --git a/Engine/Source/Developer/TraceTools/Private/Widgets/STraceStatistics.h b/Engine/Source/Developer/TraceTools/Private/Widgets/STraceStatistics.h
--- a/Engine/Source/Developer/TraceTools/Private/Widgets/STraceStatistics.h
+++ b/Engine/Source/Developer/TraceTools/Private/Widgets/STraceStatistics.h
@@ -6,6 +6,8 @@
 #include "Widgets/DeclarativeSyntaxSupport.h"
 #include "Widgets/SCompoundWidget.h"
 
+struct EVisibility;
+class FReply;
 class ITraceController;
 
 namespace UE::TraceTools
@@ -35,6 +37,10 @@ private:
 	FText GetStatsMemoryValueText(uint64 InValue) const;
 	FText GetStatsBandwidthText(uint64 InValue) const;
 	FText GetStatsCacheText() const;
+	FText GetTraceEndpointText() const;
+
+	FReply CopyEndpoint_OnClicked() const;
+	EVisibility GetCopyEndpointVisibility() const;
 
 private:
 	TSharedPtr<ISessionTraceFilterService> SessionFilterService;
diff --git a/Engine/Source/Developer/TraceTools/TraceTools.build.cs b/Engine/Source/Developer/TraceTools/TraceTools.build.cs
--- a/Engine/Source/Developer/TraceTools/TraceTools.build.cs
+++ b/Engine/Source/Developer/TraceTools/TraceTools.build.cs
@@ -10,6 +10,7 @@ public class TraceTools : ModuleRules
 
 		PrivateDependencyModuleNames.AddRange(
 				new string[] {
+					"ApplicationCore",
 					"Core",
 					"CoreUObject",
 					"InputCore",
