Add trace status information to the Trace Control Window.
Fix trace stats bandwidth calculations by using runtime timestamps.
Get Engine presets from the runtime via Message Bus in Trace Control Window.
Fix crash when UE_TRACE_ENABLED is off.

diff --git a/Engine/Source/Developer/TraceTools/Private/Models/ITraceFilterPreset.h b/Engine/Source/Developer/TraceTools/Private/Models/ITraceFilterPreset.h
--- a/Engine/Source/Developer/TraceTools/Private/Models/ITraceFilterPreset.h
+++ b/Engine/Source/Developer/TraceTools/Private/Models/ITraceFilterPreset.h
@@ -25,6 +25,8 @@ public:
 	virtual bool CanDelete() const = 0;
 	/** Whether or not this preset only exists locally */
 	virtual bool IsLocal() const = 0;
+	/** Whether or not this preset is a runtime preset */
+	virtual bool IsEnginePreset() const = 0;
 
 	/** Update this preset according to the provided set of trace objects */
 	virtual void Save(const TArray<TSharedPtr<ITraceObject>>& InObjects) = 0;
diff --git a/Engine/Source/Developer/TraceTools/Private/Models/TraceFilterPresets.h b/Engine/Source/Developer/TraceTools/Private/Models/TraceFilterPresets.h
--- a/Engine/Source/Developer/TraceTools/Private/Models/TraceFilterPresets.h
+++ b/Engine/Source/Developer/TraceTools/Private/Models/TraceFilterPresets.h
@@ -101,6 +101,7 @@ public:
 	virtual bool MakeShared() override;
 	virtual bool MakeLocal() override;
 	virtual bool IsLocal() const override;
+	virtual bool IsEnginePreset() const override { return false; }
 	virtual void Save(const TArray<TSharedPtr<TraceTools::ITraceObject>>& InObjects) override {}
 	virtual void Save() override {}
 	/** End ITraceFilterPreset overrides */
@@ -124,6 +125,7 @@ public:
 	virtual bool MakeShared() override { return false; }
 	virtual bool MakeLocal() override { return false; }
 	virtual bool IsLocal() const override { return false; }
+	virtual bool IsEnginePreset() const override { return true; }
 	virtual void Save(const TArray<TSharedPtr<TraceTools::ITraceObject>>& InObjects) override {}
 	virtual void Save() override {}
 	/** End ITraceFilterPreset overrides */
diff --git a/Engine/Source/Developer/TraceTools/Private/Services/ISessionTraceFilterService.h b/Engine/Source/Developer/TraceTools/Private/Services/ISessionTraceFilterService.h
--- a/Engine/Source/Developer/TraceTools/Private/Services/ISessionTraceFilterService.h
+++ b/Engine/Source/Developer/TraceTools/Private/Services/ISessionTraceFilterService.h
@@ -52,8 +52,8 @@ public:
 	/** Set the filtered state for an individual object by its hash. */
 	virtual void SetObjectFilterState(const FString& InObjectName, const bool bFilterState) = 0;
 	
-	/** Set timestamp for last processed update (data change). */
-	virtual const FDateTime& GetTimestamp() const = 0;
+	/** Get the timestamp for the last channel data update. */
+	virtual const FDateTime& GetChannelsUpdateTimestamp() const = 0;
 
 	/** Update filter preset */
 	virtual void UpdateFilterPreset(const TSharedPtr<ITraceFilterPreset> InPreset, bool IsEnabled) = 0;
@@ -72,6 +72,9 @@ public:
 
 	/** Get the endpoint of the current running trace. */
 	virtual const FString& GetTraceEndpoint() const = 0;
+
+	/** Get the current status of the trace system. */
+	virtual FTraceStatus::ETraceSystemStatus GetTraceSystemStatus() const = 0;
 };
 
 } // namespace UE::TraceTools
\ No newline at end of file
diff --git a/Engine/Source/Developer/TraceTools/Private/Services/SessionTraceControllerFilterService.cpp b/Engine/Source/Developer/TraceTools/Private/Services/SessionTraceControllerFilterService.cpp
--- a/Engine/Source/Developer/TraceTools/Private/Services/SessionTraceControllerFilterService.cpp
+++ b/Engine/Source/Developer/TraceTools/Private/Services/SessionTraceControllerFilterService.cpp
@@ -51,9 +51,9 @@ const FTraceObjectInfo* FSessionTraceControllerFilterService::GetObject(const FS
 	return Objects.Find(HashName(Name));
 }
 
-const FDateTime& FSessionTraceControllerFilterService::GetTimestamp() const
+const FDateTime& FSessionTraceControllerFilterService::GetChannelsUpdateTimestamp() const
 {
-	return TimeStamp;
+	return ChannelsTimestamp;
 }
 
 void FSessionTraceControllerFilterService::SetObjectFilterState(const FString& InObjectName, const bool bFilterState)
@@ -104,10 +104,8 @@ void FSessionTraceControllerFilterService::DisableAllChannels()
 
 void FSessionTraceControllerFilterService::OnTraceStatusUpdated(const FTraceStatus& InStatus, FTraceStatus::EUpdateType InUpdateType, ITraceControllerCommands& Commands)
 {
-	double DeltaTimeSeconds = (FDateTime::Now() - TimeStamp).GetTotalSeconds();
 	if (!TraceController->HasAvailableSelectedInstance())
 	{
-		TimeStamp = FDateTime::Now();
 		Objects.Empty();
 		return;
 	}
@@ -126,10 +124,13 @@ void FSessionTraceControllerFilterService::OnTraceStatusUpdated(const FTraceStat
 	if (EnumHasAnyFlags(InUpdateType, FTraceStatus::EUpdateType::Status))
 	{
 		TraceEndpoint = InStatus.Endpoint;
+		TraceSystemStatus = InStatus.TraceSystemStatus;
 
 		Stats.BytesSentPerSecond = 0;
 		Stats.BytesTracedPerSecond = 0;
 
+		double DeltaTimeSeconds = (InStatus.StatusTimestamp - StatusTimestamp).GetTotalSeconds();
+
 		if (DeltaTimeSeconds > 0.0f)
 		{
 			Stats.BytesSentPerSecond = FMath::Max(0ull, InStatus.Stats.BytesSent - Stats.StandardStats.BytesSent);
@@ -140,13 +141,14 @@ void FSessionTraceControllerFilterService::OnTraceStatusUpdated(const FTraceStat
 		}
 
 		Stats.StandardStats = InStatus.Stats;
+		StatusTimestamp = InStatus.StatusTimestamp;
 		bHasStats = true;
 	}
 }
 
 void FSessionTraceControllerFilterService::UpdateChannels(const FTraceStatus& InStatus)
 {
-	TimeStamp = FDateTime::Now();
+	ChannelsTimestamp = FDateTime::Now();
 
 	const TMap<uint32, FTraceStatus::FChannel> Channels = InStatus.Channels;
 	Objects.Empty(Channels.Num());
diff --git a/Engine/Source/Developer/TraceTools/Private/Services/SessionTraceControllerFilterService.h b/Engine/Source/Developer/TraceTools/Private/Services/SessionTraceControllerFilterService.h
--- a/Engine/Source/Developer/TraceTools/Private/Services/SessionTraceControllerFilterService.h
+++ b/Engine/Source/Developer/TraceTools/Private/Services/SessionTraceControllerFilterService.h
@@ -27,7 +27,7 @@ public:
 	virtual void GetRootObjects(TArray<FTraceObjectInfo>& OutObjects) const override;
 	virtual const FTraceObjectInfo* GetObject(const FString& Name) const override;
 	
-	virtual const FDateTime& GetTimestamp() const override;
+	virtual const FDateTime& GetChannelsUpdateTimestamp() const override;
 	virtual void SetObjectFilterState(const FString& InObjectName, const bool bFilterState) override;
 	virtual void UpdateFilterPreset(const TSharedPtr<ITraceFilterPreset> InPreset, bool IsEnabled) override;
 
@@ -38,6 +38,7 @@ public:
 	virtual const FTraceStats& GetStats() const override;
 
 	virtual const FString& GetTraceEndpoint() const override { return TraceEndpoint; }
+	virtual FTraceStatus::ETraceSystemStatus GetTraceSystemStatus() const { return TraceSystemStatus; }
 	/** End ISessionTraceFilterService overrides */
 
 protected:
@@ -63,8 +64,11 @@ protected:
 	TSet<FString> FrameEnabledChannels;
 	TSet<FString> FrameDisabledChannels;
 
-	/** Timestamp at which contained data (including provider) was last updated */
-	FDateTime TimeStamp;
+	/** Timestamp at which the trace status was last updated */
+	FDateTime StatusTimestamp;
+
+	/** Timestamp at which the trace channels were last updated */
+	FDateTime ChannelsTimestamp;
 
 	bool bChannelsReceived = false;
 
@@ -75,6 +79,7 @@ protected:
 	bool bHasStats = false;
 
 	FString TraceEndpoint;
+	FTraceStatus::ETraceSystemStatus TraceSystemStatus;
 };
 
 } // namespace UE::TraceTools
\ No newline at end of file
diff --git a/Engine/Source/Developer/TraceTools/Private/Widgets/SFilterPresetList.cpp b/Engine/Source/Developer/TraceTools/Private/Widgets/SFilterPresetList.cpp
--- a/Engine/Source/Developer/TraceTools/Private/Widgets/SFilterPresetList.cpp
+++ b/Engine/Source/Developer/TraceTools/Private/Widgets/SFilterPresetList.cpp
@@ -85,7 +85,7 @@ void SFilterPresetList::RefreshFilterPresets()
 		for (auto FilterIt = Presets.CreateConstIterator(); FilterIt; ++FilterIt)
 		{
 			const TSharedRef<SFilterPreset>& FilterToRemove = *FilterIt;
-			ActiveFilterNames.Add(FilterToRemove->GetFilterPreset()->GetName());
+			CurrentActiveFilterNames.Add(FilterToRemove->GetFilterPreset()->GetName());
 		}
 
 		FilterBox->ClearChildren();
@@ -109,15 +109,14 @@ void SFilterPresetList::RefreshFilterPresets()
 
 	for (const TSharedPtr<ITraceFilterPreset>& Preset : AllFilterPresets)
 	{
-		if (ActiveFilterNames.Contains(Preset->GetName()))
+		if (CurrentActiveFilterNames.Contains(Preset->GetName()))
 		{
 			TSharedRef<SFilterPreset> Filter = AddFilterPreset(Preset);
+			CurrentActiveFilterNames.Remove(Preset->GetName());
 		}
 	}
 
 	RefreshPresetEnabledState();
-
-	ActiveFilterNames.Empty();
 }
 
 FReply SFilterPresetList::OnMouseButtonUp( const FGeometry& MyGeometry, const FPointerEvent& MouseEvent )
@@ -146,7 +145,6 @@ void SFilterPresetList::SaveSettings(const FString& IniFilename) const
 {
 	FStringView IniSectionName = TEXTVIEW("Trace.FilterPresetList");
 	FStringView IniActivePresetsKey = TEXTVIEW("ActivePresets");
-	FStringView IniEnabledPresetsKey = TEXTVIEW("EnabledPresets");
 
 	FStringView Separator = TEXTVIEW(";");
 	
@@ -169,15 +167,13 @@ void SFilterPresetList::LoadSettings(const FString& IniFilename)
 {
 	FStringView IniSectionName = TEXTVIEW("Trace.FilterPresetList");
 	FStringView IniActivePresetsKey = TEXT("ActivePresets");
-	FStringView IniEnabledPresetsKey = TEXT("EnabledPresets");
 
 	FString ActivePresetsString;
 	FString EnabledPresetsString;
 	GConfig->GetString(IniSectionName.GetData(), IniActivePresetsKey.GetData(), ActivePresetsString, IniFilename);
-	GConfig->GetString(IniSectionName.GetData(), IniEnabledPresetsKey.GetData(), EnabledPresetsString, IniFilename);
 
 	FStringView Separator = TEXT(";");
-	ActivePresetsString.ParseIntoArray(ActiveFilterNames, Separator.GetData());
+	ActivePresetsString.ParseIntoArray(CurrentActiveFilterNames, Separator.GetData());
 }
 
 bool SFilterPresetList::HasAnyPresets() const
@@ -544,30 +540,18 @@ TSharedRef<SWidget> SFilterPresetList::ExternalMakeFilterPresetsMenu()
 
 void SFilterPresetList::LoadEnginePresets()
 {
-	TArray<FString> PresetStrings;
-	GConfig->GetSection(TEXT("Trace.ChannelPresets"), PresetStrings, GEngineIni);
-
-	for (const FString& Item : PresetStrings)
-	{
-		FString Key, Value;
-		Item.Split(TEXT("="), &Key, &Value);
-
-		TArray<FString> AllowListedNames;
-		Value.ParseIntoArray(AllowListedNames, TEXT(","));
-		EngineFilterPresets.Add(MakeShared<FEngineFilterPreset>(Key, AllowListedNames));
-	}
-
-	TArray<FTraceAuxiliary::FChannelPreset> EnginePresets;
-	FTraceAuxiliary::GetFixedChannelPresets(EnginePresets);
-
-	for (auto& Item : EnginePresets)
+	if (SessionFilterService->HasSettings())
 	{
-		if (!Item.bIsReadOnly)
+		const TArray<FTraceStatus::FChannelPreset>& RuntimePresets = SessionFilterService->GetSettings().ChannelPresets;
+		
+		for (const FTraceStatus::FChannelPreset& Preset : RuntimePresets)
 		{
-			TArray<FString> AllowListedNames;
-			FString Channels = Item.Channels;
-			Channels.ParseIntoArray(AllowListedNames, TEXT(","));
-			EngineFilterPresets.Add(MakeShared<FEngineFilterPreset>(Item.Name, AllowListedNames));
+			if (!Preset.bIsReadOnly)
+			{
+				TArray<FString> AllowListedNames;
+				Preset.ChannelList.ParseIntoArray(AllowListedNames, TEXT(","));
+				EngineFilterPresets.Add(MakeShared<FEngineFilterPreset>(Preset.Name, AllowListedNames));
+			}
 		}
 	}
 }
diff --git a/Engine/Source/Developer/TraceTools/Private/Widgets/SFilterPresetList.h b/Engine/Source/Developer/TraceTools/Private/Widgets/SFilterPresetList.h
--- a/Engine/Source/Developer/TraceTools/Private/Widgets/SFilterPresetList.h
+++ b/Engine/Source/Developer/TraceTools/Private/Widgets/SFilterPresetList.h
@@ -52,10 +52,10 @@ public:
 	/** Sets channels as enabled/disabled based on the current object list */
 	void RefreshPresetEnabledState();
 
-private:
 	/** Refreshes the filter preset objects and child widgets */
 	void RefreshFilterPresets();
 
+private:
 	/** Returns true if any Presets are applied */
 	bool HasAnyPresets() const;
 
@@ -130,8 +130,8 @@ private:
 	/** Delegate for when a Preset is to be highlighted */
 	FOnHighlightPreset OnHighlightPreset;
 
-	/** Cached arrays of currently active and enabled filter names, used for saving/restoring state during refreshes */
-	TArray<FString> ActiveFilterNames;
+	/** Cached array of currently active filter names, used for saving/restoring state during refreshes */
+	TArray<FString> CurrentActiveFilterNames;
 
 	/** A pointer to the SessionFilterService to autodetect enabled presets. */
 	TSharedPtr<ISessionTraceFilterService> SessionFilterService;
diff --git a/Engine/Source/Developer/TraceTools/Private/Widgets/STraceDataFilterWidget.cpp b/Engine/Source/Developer/TraceTools/Private/Widgets/STraceDataFilterWidget.cpp
--- a/Engine/Source/Developer/TraceTools/Private/Widgets/STraceDataFilterWidget.cpp
+++ b/Engine/Source/Developer/TraceTools/Private/Widgets/STraceDataFilterWidget.cpp
@@ -492,7 +492,7 @@ void STraceDataFilterWidget::RefreshTileViewData()
 {
 	if (SessionFilterService.IsValid())
 	{
-		SyncTimeStamp = SessionFilterService->GetTimestamp();
+		SyncTimeStamp = SessionFilterService->GetChannelsUpdateTimestamp();
 
 		/** Save expansion and selection */
 		SaveItemSelection();
@@ -521,7 +521,13 @@ void STraceDataFilterWidget::Tick(const FGeometry& AllottedGeometry, const doubl
 
 	if (SessionFilterService.IsValid() )
 	{
-		if (SessionFilterService->GetTimestamp() != SyncTimeStamp)
+		if (bHasSettings == false && SessionFilterService->HasSettings())
+		{
+			FilterPresetsListWidget->RefreshFilterPresets();
+			bHasSettings = true;
+		}
+
+		if (SessionFilterService->GetChannelsUpdateTimestamp() != SyncTimeStamp)
 		{
 			RefreshTileViewData();
 			bNeedsListRefresh = true;
@@ -565,6 +571,7 @@ void STraceDataFilterWidget::Tick(const FGeometry& AllottedGeometry, const doubl
 void STraceDataFilterWidget::OnSessionSelectionChanged()
 {
 	bHasChannelData = false;
+	bHasSettings = false;
 }
 
 } // namespace UE::TraceTools
diff --git a/Engine/Source/Developer/TraceTools/Private/Widgets/STraceDataFilterWidget.h b/Engine/Source/Developer/TraceTools/Private/Widgets/STraceDataFilterWidget.h
--- a/Engine/Source/Developer/TraceTools/Private/Widgets/STraceDataFilterWidget.h
+++ b/Engine/Source/Developer/TraceTools/Private/Widgets/STraceDataFilterWidget.h
@@ -126,6 +126,7 @@ protected:
 	double AccumulatedTime = 0.0f;
 
 	bool bHasChannelData = false;
+	bool bHasSettings = false;
 };
 
 } // namespace UE::TraceTools
\ No newline at end of file
diff --git a/Engine/Source/Developer/TraceTools/Private/Widgets/STraceStatistics.cpp b/Engine/Source/Developer/TraceTools/Private/Widgets/STraceStatistics.cpp
--- a/Engine/Source/Developer/TraceTools/Private/Widgets/STraceStatistics.cpp
+++ b/Engine/Source/Developer/TraceTools/Private/Widgets/STraceStatistics.cpp
@@ -6,6 +6,7 @@
 #include "Internationalization/Text.h"
 #include "SlateOptMacros.h"
 #include "Styling/StyleColors.h"
+#include "ProfilingDebugging/TraceAuxiliary.h"
 #include "Widgets/Images/SImage.h"
 #include "Widgets/Input/SButton.h"
 #include "Widgets/Layout/SBorder.h"
@@ -47,6 +48,34 @@ void STraceStatistics::Construct(const FArguments& InArgs, TSharedPtr<ISessionTr
 			[
 				SNew(SVerticalBox)
 				
+				+ SVerticalBox::Slot()
+				.HAlign(EHorizontalAlignment::HAlign_Left)
+				.Padding(0.0f, 3.0f, 0.0f, 0.0f)
+				.AutoHeight()
+				[
+					SNew(SHorizontalBox)
+							
+					+ SHorizontalBox::Slot()
+					.Padding(0.0f, 2.0f, 0.0f, 0.0f)
+					.AutoWidth()
+					[
+						SNew(STextBlock)
+						.ColorAndOpacity(FSlateColor(EStyleColor::Foreground))
+						.ToolTipText(LOCTEXT("TraceStatusTooltip", "The status of the tracing system."))
+						.Text(LOCTEXT("TraceStatus", "Trace Status:"))
+					]
+
+					+ SHorizontalBox::Slot()
+					.Padding(2.0f, 2.0f, 0.0f, 0.0f)
+					.AutoWidth()
+					[
+						SNew(STextBlock)
+						.ColorAndOpacity(FSlateColor(EStyleColor::AccentGray))
+						.Text(this, &STraceStatistics::GetTraceSystemStateText)
+						.ToolTipText(this, &STraceStatistics::GetTraceSystemStateTooltipText)
+					]
+				]
+
 				+ SVerticalBox::Slot()
 				.HAlign(EHorizontalAlignment::HAlign_Left)
 				.Padding(0.0f, 3.0f, 0.0f, 0.0f)
@@ -424,6 +453,73 @@ FText STraceStatistics::GetTraceEndpointText() const
 	return FText::FromString(SessionFilterService->GetTraceEndpoint());
 }
 
+FText STraceStatistics::GetTraceSystemStateText() const
+{
+	// If you update these values also check GetTraceSystemStateTooltipText.
+	static_assert((uint8)FTraceStatus::ETraceSystemStatus::NotAvailable == (uint8) FTraceAuxiliary::ETraceSystemStatus::NotAvailable);
+	static_assert((uint8)FTraceStatus::ETraceSystemStatus::Available == (uint8) FTraceAuxiliary::ETraceSystemStatus::Available);
+	static_assert((uint8)FTraceStatus::ETraceSystemStatus::TracingToFile == (uint8) FTraceAuxiliary::ETraceSystemStatus::TracingToFile);
+	static_assert((uint8)FTraceStatus::ETraceSystemStatus::TracingToServer == (uint8) FTraceAuxiliary::ETraceSystemStatus::TracingToServer);
+	static_assert((uint8)FTraceStatus::ETraceSystemStatus::NumValues == (uint8) FTraceAuxiliary::ETraceSystemStatus::NumValues, "ETraceSystemStatus enum values are of out sync.");
+
+	if (!SessionFilterService->HasStats())
+	{
+		return LOCTEXT("N/A", "N/A");
+	}
+
+	switch (SessionFilterService->GetTraceSystemStatus())
+	{
+	case FTraceStatus::ETraceSystemStatus::NotAvailable:
+	{
+		return LOCTEXT("TraceSystemNotAvailableText", "Not Available");
+	}
+	case FTraceStatus::ETraceSystemStatus::Available:
+	{
+		return LOCTEXT("TraceSystemAvailableText", "Available");
+	}
+	case FTraceStatus::ETraceSystemStatus::TracingToServer:
+	{
+		return LOCTEXT("TracingToServerText", "Tracing to Server");
+	}
+	case FTraceStatus::ETraceSystemStatus::TracingToFile:
+	{
+		return LOCTEXT("TracingToFileText", "Tracing to File");
+	}
+	default:
+		return LOCTEXT("Unkown", "Unkown");
+	}
+}
+
+FText STraceStatistics::GetTraceSystemStateTooltipText() const
+{
+	if (!SessionFilterService->HasStats())
+	{
+		return FText::GetEmpty();
+	}
+
+	switch (SessionFilterService->GetTraceSystemStatus())
+	{
+	case FTraceStatus::ETraceSystemStatus::NotAvailable:
+	{
+		return LOCTEXT("TraceSystemNotAvailableTooltipText", "Trace system is disabled at compile time. Check the UE_TRACE_ENABLED define.");
+	}
+	case FTraceStatus::ETraceSystemStatus::Available:
+	{
+		return LOCTEXT("TraceSystemAvailableTooltipText", "Trace system is available and can be started. Data might be stored in the Important Events and Tail buffers.");
+	}
+	case FTraceStatus::ETraceSystemStatus::TracingToServer:
+	{
+		return LOCTEXT("TracingToServerTooltipText", "Tracing to the trace server.");
+	}
+	case FTraceStatus::ETraceSystemStatus::TracingToFile:
+	{
+		return LOCTEXT("TracingToFileTooltipText", "Tracing directly to a file.");
+	}
+	default:
+		return FText::GetEmpty();
+	}
+}
+
 FReply STraceStatistics::CopyEndpoint_OnClicked() const
 {
 	const FString& Endpoint = SessionFilterService->GetTraceEndpoint();
diff --git a/Engine/Source/Developer/TraceTools/Private/Widgets/STraceStatistics.h b/Engine/Source/Developer/TraceTools/Private/Widgets/STraceStatistics.h
--- a/Engine/Source/Developer/TraceTools/Private/Widgets/STraceStatistics.h
+++ b/Engine/Source/Developer/TraceTools/Private/Widgets/STraceStatistics.h
@@ -38,6 +38,8 @@ private:
 	FText GetStatsBandwidthText(uint64 InValue) const;
 	FText GetStatsCacheText() const;
 	FText GetTraceEndpointText() const;
+	FText GetTraceSystemStateText() const;
+	FText GetTraceSystemStateTooltipText() const;
 
 	FReply CopyEndpoint_OnClicked() const;
 	EVisibility GetCopyEndpointVisibility() const;
diff --git a/Engine/Source/Runtime/Core/Private/ProfilingDebugging/TraceAuxiliary.cpp b/Engine/Source/Runtime/Core/Private/ProfilingDebugging/TraceAuxiliary.cpp
--- a/Engine/Source/Runtime/Core/Private/ProfilingDebugging/TraceAuxiliary.cpp
+++ b/Engine/Source/Runtime/Core/Private/ProfilingDebugging/TraceAuxiliary.cpp
@@ -281,15 +281,15 @@ void FTraceAuxiliaryImpl::ForEachChannel(const TCHAR* ChannelList, bool bResolve
 			// Check against hard coded presets
 			if (FCString::Stricmp(Name, GDefaultChannels.Name) == 0)
 			{
-				ForEachChannel(GDefaultChannels.Channels, false, LogCategory, Callable);
+				ForEachChannel(GDefaultChannels.ChannelList, false, LogCategory, Callable);
 			}
 			else if (FCString::Stricmp(Name,GMemoryChannels.Name) == 0)
 			{
-				ForEachChannel(GMemoryChannels.Channels, false, LogCategory, Callable);
+				ForEachChannel(GMemoryChannels.ChannelList, false, LogCategory, Callable);
 			}
 			else if (FCString::Stricmp(Name, GMemoryLightChannels.Name) == 0)
 			{
-				ForEachChannel(GMemoryLightChannels.Channels, false, LogCategory, Callable);
+				ForEachChannel(GMemoryLightChannels.ChannelList, false, LogCategory, Callable);
 			}
 			// Check against data driven presets (if available)
 			else if (GConfig && GConfig->GetString(TEXT("Trace.ChannelPresets"), Name, Value, GEngineIni))
@@ -1330,12 +1330,12 @@ static bool StartFromCommandlineArguments(const TCHAR* CommandLine, bool& bOutSt
 	}
 	else if (FParse::Param(CommandLine, TEXT("trace")))
 	{
-		Channels = GDefaultChannels.Channels;
+		Channels = GDefaultChannels.ChannelList;
 	}
 #if WITH_EDITOR
 	else
 	{
-		Channels = GDefaultChannels.Channels;
+		Channels = GDefaultChannels.ChannelList;
 	}
 #endif
 
@@ -1388,7 +1388,7 @@ static bool StartFromCommandlineArguments(const TCHAR* CommandLine, bool& bOutSt
 	// If user has defined a connection type but not specified channels, use the default channel set.
 	if (Type != FTraceAuxiliary::EConnectionType::None && Channels.IsEmpty())
 	{
-		Channels = GDefaultChannels.Channels;
+		Channels = GDefaultChannels.ChannelList;
 	}
 
 	if (Channels.IsEmpty())
@@ -1851,11 +1851,61 @@ UE::Trace::FInitializeDesc const* FTraceAuxiliary::GetInitializeDesc()
 #endif
 }
 
-void FTraceAuxiliary::GetFixedChannelPresets(TArray<FChannelPreset>& OutPresets)
+void FTraceAuxiliary::EnumerateFixedChannelPresets(PresetCallback Callback)
 {
-	OutPresets.Emplace(GDefaultChannels);
-	OutPresets.Emplace(GMemoryChannels);
-	OutPresets.Emplace(GMemoryLightChannels);
+	if (Callback(GDefaultChannels) == EEnumerateResult::Stop)
+	{
+		return;
+	}
+
+	if (Callback(GMemoryChannels) == EEnumerateResult::Stop)
+	{
+		return;
+	}
+
+	if (Callback(GMemoryLightChannels) == EEnumerateResult::Stop)
+	{
+		return;
+	}
+}
+
+void FTraceAuxiliary::EnumerateChannelPresetsFromSettings(PresetCallback Callback)
+{
+	TArray<FString> PresetStrings;
+	GConfig->GetSection(TEXT("Trace.ChannelPresets"), PresetStrings, GEngineIni);
+
+	for (const FString& Item : PresetStrings)
+	{
+		FString Key, Value;
+		Item.Split(TEXT("="), &Key, &Value);
+
+		FChannelPreset Preset(*Key, *Value, false);
+		if (Callback(Preset) == EEnumerateResult::Stop)
+		{
+			return;
+		}
+	}
+}
+
+FTraceAuxiliary::ETraceSystemStatus FTraceAuxiliary::GetTraceSystemStatus()
+{
+#if UE_TRACE_ENABLED
+	EConnectionType ConnectionType = GetConnectionType();
+	if (ConnectionType == EConnectionType::Network)
+	{
+		return ETraceSystemStatus::TracingToServer;
+	}
+	else if (ConnectionType == EConnectionType::File)
+	{
+		return ETraceSystemStatus::TracingToFile;
+	}
+	else
+	{
+		return ETraceSystemStatus::Available;
+	}
+#else
+	return  ETraceSystemStatus::NotAvailable;
+#endif
 }
 
 ////////////////////////////////////////////////////////////////////////////////
diff --git a/Engine/Source/Runtime/Core/Public/ProfilingDebugging/TraceAuxiliary.h b/Engine/Source/Runtime/Core/Public/ProfilingDebugging/TraceAuxiliary.h
--- a/Engine/Source/Runtime/Core/Public/ProfilingDebugging/TraceAuxiliary.h
+++ b/Engine/Source/Runtime/Core/Public/ProfilingDebugging/TraceAuxiliary.h
@@ -19,25 +19,51 @@
 class FTraceAuxiliary
 {
 public:
+	/** 
+	* This enum is serialized and sent via the trace service.
+	* Do not change the values or modify the order. Only add new values to the end.
+	* Should be kept in sync with FTraceStatus::ETraceSystemStatus from ITraceController.h
+	*/
+	enum class ETraceSystemStatus : uint8
+	{
+		NotAvailable, // Disabled at compile time.
+		Available,
+		TracingToServer,
+		TracingToFile,
+
+		NumValues, // This must be the last value.
+	};
+
+	enum class EEnumerateResult : uint8
+	{
+		Continue,
+		Stop,
+	};
+
 	struct FChannelPreset
 	{
-		FChannelPreset(const TCHAR* InName, const TCHAR* InChannels, bool bInIsReadOnly)
+		FChannelPreset(const TCHAR* InName,const TCHAR* InChannels, bool bInIsReadOnly)
 			: Name(InName)
-			, Channels(InChannels)
+			, ChannelList(InChannels)
 			, bIsReadOnly(bInIsReadOnly)
 		{
 		}
 
+		/**
+		 * Do not store these pointers.
+		 */
 		const TCHAR* Name;
-		const TCHAR* Channels;
+		const TCHAR* ChannelList;
 
 		/**
-		* A preset should be read-only if it contains any read-only channels. 
+		* A preset should be read-only if it contains any read-only channels.
 		* A read-only preset can only be enabled using the command line when starting the application.
 		*/
 		bool bIsReadOnly = false;
 	};
 
+	typedef TFunctionRef<EEnumerateResult(const FChannelPreset& Preset)> PresetCallback;
+
 	// In no logging configurations all log categories are of type FNoLoggingCategory, which has no relation with
 	// FLogCategoryBase. In order to not need to conditionally set the argument alias the type here.
 #if NO_LOGGING
@@ -236,9 +262,14 @@ public:
 	static CORE_API struct UE::Trace::FInitializeDesc const* GetInitializeDesc();
 
 	/**
-	* Get the channel presets that are defined in code.
+	* Enumerate the channel presets that are defined in code.
+	*/
+	static CORE_API void EnumerateFixedChannelPresets(PresetCallback Callback);
+
+	/**
+	* Enumerate the channel presets that are defined in BaseEngine.ini, under the [Trace.ChannelPresets] section.
 	*/
-	static CORE_API void GetFixedChannelPresets(TArray<FChannelPreset>& OutPresets);
+	static CORE_API void EnumerateChannelPresetsFromSettings(PresetCallback Callback);
 
 	/**
 	 * Delegate that triggers when a connection is established. Gives subscribers a chance to trace events that appear
@@ -270,6 +301,11 @@ public:
 	 * The path to the snapshot file is passed to the delegate.
 	 */
 	static CORE_API FOnSnapshotSaved OnSnapshotSaved;
+
+	/**
+	 * Returns the current status of the trace system.
+	 */
+	static CORE_API ETraceSystemStatus GetTraceSystemStatus();
 };
 
 #if UE_TRACE_SERVER_CONTROLS_ENABLED
diff --git a/Engine/Source/Runtime/Engine/Private/TraceService.cpp b/Engine/Source/Runtime/Engine/Private/TraceService.cpp
--- a/Engine/Source/Runtime/Engine/Private/TraceService.cpp
+++ b/Engine/Source/Runtime/Engine/Private/TraceService.cpp
@@ -110,6 +110,8 @@ void FTraceServiceImpl::FillTraceStatusMessage(FTraceControlStatus* Message)
 	Message->CacheWaste = Stats.CacheWaste;
 	Message->bAreStatNamedEventsEnabled = GCycleStatsShouldEmitNamedEvents > 0;
 	Message->bIsPaused = FTraceAuxiliary::IsPaused();
+	Message->StatusTimestamp = FDateTime::Now();
+	Message->TraceSystemStatus = static_cast<uint8>(FTraceAuxiliary::GetTraceSystemStatus());
 }
 
 void FTraceServiceImpl::OnChannelSet(const FTraceControlChannelsSet& Message, const TSharedRef<IMessageContext>& Context)
@@ -253,10 +255,27 @@ void FTraceServiceImpl::OnSettingsPing(const FTraceControlSettingsPing& Message,
 	FTraceControlSettings* Response = FMessageEndpoint::MakeMessage<FTraceControlSettings>();
 	UE::Trace::FInitializeDesc const* InitDesc = FTraceAuxiliary::GetInitializeDesc();
 
-	Response->bUseImportantCache = InitDesc->bUseImportantCache;
-	Response->bUseWorkerThread = InitDesc->bUseWorkerThread;
-	Response->TailSizeBytes = InitDesc->TailSizeBytes;
-	
+	if (InitDesc)
+	{
+		Response->bUseImportantCache = InitDesc->bUseImportantCache;
+		Response->bUseWorkerThread = InitDesc->bUseWorkerThread;
+		Response->TailSizeBytes = InitDesc->TailSizeBytes;
+	}
+
+	auto AddPreset = [&Response](const FTraceAuxiliary::FChannelPreset& Preset)
+	{
+		FTraceChannelPreset TracePreset;
+		TracePreset.Name = Preset.Name;
+		TracePreset.ChannelList = Preset.ChannelList;
+		TracePreset.bIsReadOnly = Preset.bIsReadOnly;
+
+		Response->ChannelPresets.Add(TracePreset);
+		return FTraceAuxiliary::EEnumerateResult::Continue;
+	};
+
+	FTraceAuxiliary::EnumerateFixedChannelPresets(AddPreset);
+	FTraceAuxiliary::EnumerateChannelPresetsFromSettings(AddPreset);
+
 	MessageEndpoint->Send(Response, Context->GetSender());
 }
 
diff --git a/Engine/Source/Runtime/EngineMessages/Public/TraceControlMessages.h b/Engine/Source/Runtime/EngineMessages/Public/TraceControlMessages.h
--- a/Engine/Source/Runtime/EngineMessages/Public/TraceControlMessages.h
+++ b/Engine/Source/Runtime/EngineMessages/Public/TraceControlMessages.h
@@ -64,6 +64,12 @@ struct FTraceControlStatus
 
 	UPROPERTY(EditAnywhere, Category = "Message")
 	bool bIsTracing = false;
+
+	UPROPERTY(EditAnywhere, Category = "Message")
+	FDateTime StatusTimestamp;
+
+	UPROPERTY(EditAnywhere, Category = "Message")
+	uint8 TraceSystemStatus = 0;
 };
 
 /**
@@ -75,6 +81,24 @@ struct FTraceControlSettingsPing
 	GENERATED_USTRUCT_BODY()
 };
 
+/**
+ *	Contains data about a trace channel preset
+ */
+USTRUCT()
+struct FTraceChannelPreset
+{
+	GENERATED_USTRUCT_BODY()
+
+	UPROPERTY(EditAnywhere, Category = "Message")
+	FString Name;
+
+	UPROPERTY(EditAnywhere, Category = "Message")
+	FString ChannelList;
+
+	UPROPERTY(EditAnywhere, Category = "Message")
+	bool bIsReadOnly;
+};
+
 /**
  * Response for FTraceControlSettingsPing containing
  * TraceLog and system settings.
@@ -92,6 +116,9 @@ struct FTraceControlSettings
 
 	UPROPERTY(EditAnywhere, Category="Message")
 	uint32 TailSizeBytes = 0;
+
+	UPROPERTY(EditAnywhere, Category = "Message")
+	TArray<FTraceChannelPreset> ChannelPresets;
 };
 
 /**
diff --git a/Engine/Source/Runtime/SessionServices/Private/TraceController.cpp b/Engine/Source/Runtime/SessionServices/Private/TraceController.cpp
--- a/Engine/Source/Runtime/SessionServices/Private/TraceController.cpp
+++ b/Engine/Source/Runtime/SessionServices/Private/TraceController.cpp
@@ -219,6 +219,12 @@ void FTraceController::OnSettings(const FTraceControlSettings& Message, const TS
 		Settings.bUseImportantCache = Message.bUseImportantCache;
 		Settings.bUseWorkerThread = Message.bUseWorkerThread;
 		Settings.TailSizeBytes = Message.TailSizeBytes;
+
+		Settings.ChannelPresets.Empty();
+		for (const FTraceChannelPreset& Preset : Message.ChannelPresets)
+		{
+			Settings.ChannelPresets.Add(FTraceStatus::FChannelPreset(Preset.Name, Preset.ChannelList, Preset.bIsReadOnly));
+		}
 		
 		StatusReceivedEvent.Broadcast(Status, FTraceStatus::EUpdateType::Settings, Instance->Commands);
 		if (SelectedInstanceIds.Contains(Status.InstanceId))
@@ -273,6 +279,8 @@ void FTraceController::OnInstanceSelectionChanged(const TSharedPtr<ISessionInsta
 
 void FTraceController::UpdateStatus(const FTraceControlStatus& Message, FTraceStatus& Status)
 {
+	Status.TraceSystemStatus = static_cast<FTraceStatus::ETraceSystemStatus>(Message.TraceSystemStatus);
+	Status.StatusTimestamp = Message.StatusTimestamp;
 	Status.bIsTracing = Message.bIsTracing;
 	Status.Endpoint = Message.Endpoint;
 	Status.SessionGuid = Message.SessionGuid;
diff --git a/Engine/Source/Runtime/SessionServices/Public/ITraceController.h b/Engine/Source/Runtime/SessionServices/Public/ITraceController.h
--- a/Engine/Source/Runtime/SessionServices/Public/ITraceController.h
+++ b/Engine/Source/Runtime/SessionServices/Public/ITraceController.h
@@ -6,9 +6,9 @@
 #include "Containers/UnrealString.h"
 #include "Delegates/Delegate.h"
 #include "ITraceControllerCommands.h"
+#include "Misc/DateTime.h"
 #include "Misc/Guid.h"
 
-
 /**
  * Describes the state of a single remote instance with Trace.
  */
@@ -25,6 +25,30 @@ struct FTraceStatus
 		ChannelsStatus	= 1 << 3,
 		All				= Status|Settings|ChannelsDesc|ChannelsStatus
 	};
+
+	enum class ETraceSystemStatus : uint8
+	{
+		NotAvailable,
+		Available,
+		TracingToServer,
+		TracingToFile,
+
+		NumValues,
+	};
+
+	struct FChannelPreset
+	{
+		FChannelPreset(const FString& InName, const FString& InChannels, bool bInIsReadOnly)
+			: Name(InName)
+			, ChannelList(InChannels)
+			, bIsReadOnly(bInIsReadOnly)
+		{
+		}
+
+		FString Name;
+		FString ChannelList;
+		bool bIsReadOnly;
+	};
 	
 	struct FSettings
 	{
@@ -34,6 +58,8 @@ struct FTraceStatus
 		bool bUseImportantCache;
 		/** Size of tail buffer */
 		uint32 TailSizeBytes;
+		/** The channel presets defined by the process. */
+		TArray<FChannelPreset> ChannelPresets;
 	};
 
 	struct FChannel
@@ -89,6 +115,10 @@ struct FTraceStatus
 	TMap<uint32, FChannel> Channels;
 	/** Stats */
 	FStats Stats;
+	/** Timestamp when the state of the trace was captured */
+	FDateTime StatusTimestamp;
+	/** The status of the trace system. */
+	ETraceSystemStatus TraceSystemStatus;
 };
 
 ENUM_CLASS_FLAGS(FTraceStatus::EUpdateType);
