Fix non unity error

diff --git a/Engine/Source/Developer/TraceTools/Private/Widgets/STraceControl.cpp b/Engine/Source/Developer/TraceTools/Private/Widgets/STraceControl.cpp
--- a/Engine/Source/Developer/TraceTools/Private/Widgets/STraceControl.cpp
+++ b/Engine/Source/Developer/TraceTools/Private/Widgets/STraceControl.cpp
@@ -6,6 +6,7 @@
 #include "ISessionInstanceInfo.h"
 #include "ISessionManager.h"
 #include "ISessionServicesModule.h"
+#include "Misc/App.h"
 #include "Modules/ModuleManager.h"
 #include "STraceDataFilterWidget.h"
 #include "Widgets/SBoxPanel.h"
