Hide close button for the Connection Tab in Unreal Insights Frontend.
Reset state of Trace Control Widget when session selection changes.

diff --git a/Engine/Source/Developer/TraceInsightsFrontend/Private/InsightsFrontend/TraceInsightsFrontendModule.cpp b/Engine/Source/Developer/TraceInsightsFrontend/Private/InsightsFrontend/TraceInsightsFrontendModule.cpp
--- a/Engine/Source/Developer/TraceInsightsFrontend/Private/InsightsFrontend/TraceInsightsFrontendModule.cpp
+++ b/Engine/Source/Developer/TraceInsightsFrontend/Private/InsightsFrontend/TraceInsightsFrontendModule.cpp
@@ -302,6 +302,7 @@ TSharedRef<SDockTab> FTraceInsightsFrontendModule::SpawnConnectionTab(const FSpa
 {
 	const TSharedRef<SDockTab> DockTab = SNew(SLockedTab)
 		.TabRole(ETabRole::MajorTab)
+		.CanEverClose(false)
 		.OnCanCloseTab_Lambda([]() { return false; }); // can't close this tab
 
 	DockTab->SetOnTabClosed(SDockTab::FOnTabClosedCallback::CreateRaw(this, &FTraceInsightsFrontendModule::OnConnectionTabClosed));
diff --git a/Engine/Source/Developer/TraceTools/Private/Services/SessionTraceControllerFilterService.cpp b/Engine/Source/Developer/TraceTools/Private/Services/SessionTraceControllerFilterService.cpp
--- a/Engine/Source/Developer/TraceTools/Private/Services/SessionTraceControllerFilterService.cpp
+++ b/Engine/Source/Developer/TraceTools/Private/Services/SessionTraceControllerFilterService.cpp
@@ -212,9 +212,13 @@ const FTraceStats& FSessionTraceControllerFilterService::GetStats() const
 
 void FSessionTraceControllerFilterService::Reset()
 {
+	Objects.Empty(Objects.Num());
+
 	bHasStats = false;
 	bHasSettings = false;
 	TraceEndpoint.Empty();
+	ChannelsTimestamp = FDateTime();
+	StatusTimestamp = FDateTime();
 }
 
 void FSessionTraceControllerFilterService::SetInstanceId(const FGuid& Id)
diff --git a/Engine/Source/Developer/TraceTools/Private/Widgets/STraceControlToolbar.cpp b/Engine/Source/Developer/TraceTools/Private/Widgets/STraceControlToolbar.cpp
--- a/Engine/Source/Developer/TraceTools/Private/Widgets/STraceControlToolbar.cpp
+++ b/Engine/Source/Developer/TraceTools/Private/Widgets/STraceControlToolbar.cpp
@@ -391,6 +391,21 @@ FText STraceControlToolbar::TogglePauseResume_GetTooltip() const
 	return LOCTEXT("PauseTraceTooltip", "Disable all the trace channels and save the channel list so they can be enabled again with the resume command.");
 }
 
+void STraceControlToolbar::SetInstanceId(const FGuid& Id)
+{
+	InstanceId = Id;
+	Reset();
+}
+
+void STraceControlToolbar::Reset()
+{
+	TraceTarget = ETraceTarget::Server;
+	bIsTracing = false;
+	bIsPaused = false;
+	bAreStatNamedEventsEnabled = false;
+	bIsTracingAvailable = false;
+}
+
 } // namespace UE::TraceTools
 
 #undef LOCTEXT_NAMESPACE
\ No newline at end of file
diff --git a/Engine/Source/Developer/TraceTools/Private/Widgets/STraceControlToolbar.h b/Engine/Source/Developer/TraceTools/Private/Widgets/STraceControlToolbar.h
--- a/Engine/Source/Developer/TraceTools/Private/Widgets/STraceControlToolbar.h
+++ b/Engine/Source/Developer/TraceTools/Private/Widgets/STraceControlToolbar.h
@@ -40,7 +40,7 @@ public:
 
 	void Construct( const FArguments& InArgs, const TSharedRef<FUICommandList>& CommandList, TSharedPtr<ITraceController> InTraceController);
 
-	void SetInstanceId(const FGuid& Id) { InstanceId = Id; }
+	void SetInstanceId(const FGuid& Id);
 
 private:
 	void BindCommands(const TSharedRef<FUICommandList>& CommandList);
@@ -80,6 +80,8 @@ private:
 
 	bool IsInstanceAvailable() const;
 
+	void Reset();
+
 private:
 	TSharedPtr<ITraceController> TraceController;
 
diff --git a/Engine/Source/Developer/TraceTools/Private/Widgets/STraceDataFilterWidget.cpp b/Engine/Source/Developer/TraceTools/Private/Widgets/STraceDataFilterWidget.cpp
--- a/Engine/Source/Developer/TraceTools/Private/Widgets/STraceDataFilterWidget.cpp
+++ b/Engine/Source/Developer/TraceTools/Private/Widgets/STraceDataFilterWidget.cpp
@@ -488,15 +488,15 @@ bool STraceDataFilterWidget::EnumerateAllItems(TFunction<bool(TSharedPtr<ITraceO
 
 void STraceDataFilterWidget::RefreshTileViewData()
 {
-	if (SessionFilterService.IsValid())
-	{
-		SyncTimeStamp = SessionFilterService->GetChannelsUpdateTimestamp();
+	SyncTimeStamp = SessionFilterService->GetChannelsUpdateTimestamp();
 
-		/** Save expansion and selection */
-		SaveItemSelection();
+	/** Save expansion and selection */
+	SaveItemSelection();
 
-		ListItems.Empty();
+	ListItems.Empty();
 
+	if (SessionFilterService.IsValid())
+	{
 		TArray<FTraceObjectInfo> RootEvents;
 		SessionFilterService->GetRootObjects(RootEvents);
 
@@ -506,11 +506,11 @@ void STraceDataFilterWidget::RefreshTileViewData()
 		{
 			TSharedRef<ITraceObject> TraceObject = AddFilterableObject(RootEvent, FString());
 		}
+	}
 
-		RestoreItemSelection();
+	RestoreItemSelection();
 
-		FilterPresetsListWidget->RefreshPresetEnabledState();
-	}
+	FilterPresetsListWidget->RefreshPresetEnabledState();
 }
 
 void STraceDataFilterWidget::Tick(const FGeometry& AllottedGeometry, const double InCurrentTime, const float InDeltaTime)
