Add UI to control Verse Insights tracing.

diff --git a/Engine/Plugins/TraceUtilities/Source/InsightsEditor/Private/InsightsEditorModule.cpp b/Engine/Plugins/TraceUtilities/Source/InsightsEditor/Private/InsightsEditorModule.cpp
--- a/Engine/Plugins/TraceUtilities/Source/InsightsEditor/Private/InsightsEditorModule.cpp
+++ b/Engine/Plugins/TraceUtilities/Source/InsightsEditor/Private/InsightsEditorModule.cpp
@@ -89,7 +89,7 @@ void FInsightsEditorModule::StartupModule()
 
 	UnrealInsightsModule.OnMajorTabCreated().AddLambda([this](const FName& InMajorTabId, TSharedPtr<FTabManager> InTabManager)
 		{
-			if (InMajorTabId != FInsightsManagerTabs::MemoryProfilerTabId)
+			if (bStartAnalysisOnInsightsWindowCreated && InMajorTabId != FInsightsManagerTabs::MemoryProfilerTabId)
 			{
 				StartTraceAnalysis();
 			}
diff --git a/Engine/Plugins/TraceUtilities/Source/InsightsEditor/Private/InsightsEditorModule.h b/Engine/Plugins/TraceUtilities/Source/InsightsEditor/Private/InsightsEditorModule.h
--- a/Engine/Plugins/TraceUtilities/Source/InsightsEditor/Private/InsightsEditorModule.h
+++ b/Engine/Plugins/TraceUtilities/Source/InsightsEditor/Private/InsightsEditorModule.h
@@ -16,9 +16,13 @@ public:
 	//~ Begin IInsightsEditorModule Interface
 	virtual void StartupModule() override;
 	virtual void ShutdownModule() override;
+	virtual void SetStartAnalysisOnInsightsWindowCreated(bool InValue) { bStartAnalysisOnInsightsWindowCreated = InValue; };
 	//~ End IInsightsEditorModule Interface
 
 private:
 	void StartTraceAnalysis();
+
+private:
+	bool bStartAnalysisOnInsightsWindowCreated = true;
 };
 } // namespace UE::InsightsEditor
\ No newline at end of file
diff --git a/Engine/Plugins/TraceUtilities/Source/InsightsEditor/Public/IInsightsEditorModule.h b/Engine/Plugins/TraceUtilities/Source/InsightsEditor/Public/IInsightsEditorModule.h
--- a/Engine/Plugins/TraceUtilities/Source/InsightsEditor/Public/IInsightsEditorModule.h
+++ b/Engine/Plugins/TraceUtilities/Source/InsightsEditor/Public/IInsightsEditorModule.h
@@ -33,5 +33,7 @@ public:
 		static const FName ModuleName = "InsightsEditor";
 		return FModuleManager::Get().IsModuleLoaded(ModuleName);
 	}
+
+	virtual void SetStartAnalysisOnInsightsWindowCreated(bool InValue) = 0;
 };
 } // namespace UE::InsightsEditor
\ No newline at end of file
diff --git a/Engine/Source/Developer/TraceInsights/Private/Insights/Widgets/SMajorTabWindow.cpp b/Engine/Source/Developer/TraceInsights/Private/Insights/Widgets/SMajorTabWindow.cpp
--- a/Engine/Source/Developer/TraceInsights/Private/Insights/Widgets/SMajorTabWindow.cpp
+++ b/Engine/Source/Developer/TraceInsights/Private/Insights/Widgets/SMajorTabWindow.cpp
@@ -281,7 +281,11 @@ void SMajorTabWindow::Construct(const FArguments& InArgs, const TSharedRef<SDock
 			.Padding(8.0f)
 			[
 				SNew(STextBlock)
+#if WITH_EDITOR
+				.Text(LOCTEXT("NoTraceOverlayText", "No trace data available."))
+#else
 				.Text(LOCTEXT("SelectTraceOverlayText", "Please select a trace."))
+#endif
 			]
 		]
 	];
diff --git a/Engine/Source/Developer/TraceInsights/Private/Insights/Widgets/STimingView.cpp b/Engine/Source/Developer/TraceInsights/Private/Insights/Widgets/STimingView.cpp
--- a/Engine/Source/Developer/TraceInsights/Private/Insights/Widgets/STimingView.cpp
+++ b/Engine/Source/Developer/TraceInsights/Private/Insights/Widgets/STimingView.cpp
@@ -263,6 +263,11 @@ void STimingView::Construct(const FArguments& InArgs, FName InViewName)
 		true);
 	RightToolbar.EndSection();
 
+	for (Timing::ITimingViewExtender* Extender : GetExtenders())
+	{
+		Extender->ExtendMenuToolbars(*this, LeftToolbar, RightToolbar);
+	}
+
 	//////////////////////////////////////////////////
 
 	ChildSlot
@@ -742,14 +747,18 @@ void STimingView::Tick(const FGeometry& AllottedGeometry, const double InCurrent
 
 	////////////////////////////////////////////////////////////////////////////////////////////////////
 
-	if (Session)
+	// Tick plugin extenders.
+	// Each extender can add/remove tracks and/or change order of tracks.
+	for (Timing::ITimingViewExtender* Extender : GetExtenders())
 	{
-		// Tick plugin extenders.
-		// Each extender can add/remove tracks and/or change order of tracks.
-		for (Timing::ITimingViewExtender* Extender : GetExtenders())
+		if (Session || Extender->ShouldAlwaysTick())
 		{
 			Extender->Tick(*this, *Session.Get());
 		}
+	}
+
+	if (Session)
+	{
 #if UE_INSIGHTS_BACKWARD_COMPATIBILITY_UE54
 PRAGMA_DISABLE_DEPRECATION_WARNINGS
 		::Insights::ITimingViewSession* CurrentTimingViewSession = (::Insights::ITimingViewSession*)(Timing::ITimingViewSession*)this;
diff --git a/Engine/Source/Developer/TraceInsights/Public/Insights/ITimingViewExtender.h b/Engine/Source/Developer/TraceInsights/Public/Insights/ITimingViewExtender.h
--- a/Engine/Source/Developer/TraceInsights/Public/Insights/ITimingViewExtender.h
+++ b/Engine/Source/Developer/TraceInsights/Public/Insights/ITimingViewExtender.h
@@ -13,6 +13,7 @@
 #endif
 
 class FMenuBuilder;
+class FToolBarBuilder;
 
 namespace TraceServices { class IAnalysisSession; }
 namespace UE::Insights { class FFilterConfigurator; }
@@ -53,6 +54,12 @@ public:
 
 	/** Allows extender to add filters to the Quick Find widget. */
 	virtual void AddQuickFindFilters(TSharedPtr<FFilterConfigurator> FilterConfigurator) {}
+
+	/** Allows extender to add functionality to the menu toolbars. */
+	virtual void ExtendMenuToolbars(ITimingViewSession& InSession, FToolBarBuilder& LeftToolbar, FToolBarBuilder& RightToolbar) {}
+
+	/** Allows extender to specify if they should receive calls to Tick() even if session is not loaded. */
+	virtual bool ShouldAlwaysTick() { return false; }
 };
 
 } // namespace UE::Insights::Timing
diff --git a/Engine/Source/Developer/TraceTools/Private/Widgets/SToggleTraceButton.cpp b/Engine/Source/Developer/TraceTools/Private/Widgets/SToggleTraceButton.cpp
--- a/Engine/Source/Developer/TraceTools/Private/Widgets/SToggleTraceButton.cpp
+++ b/Engine/Source/Developer/TraceTools/Private/Widgets/SToggleTraceButton.cpp
@@ -28,10 +28,12 @@ SLATE_IMPLEMENT_WIDGET(SToggleTraceButton)
 void SToggleTraceButton::PrivateRegisterAttributes(FSlateAttributeInitializer& AttributeInitializer)
 {
 	SLATE_ADD_MEMBER_ATTRIBUTE_DEFINITION_WITH_NAME(AttributeInitializer, "IsTraceRunning", IsTraceRunningAttribute, EInvalidateWidgetReason::Paint);
+	SLATE_ADD_MEMBER_ATTRIBUTE_DEFINITION_WITH_NAME(AttributeInitializer, "DynamicToolTipText", DynamicToolTipTextAttribute, EInvalidateWidgetReason::Paint);
 }
 
 SToggleTraceButton::SToggleTraceButton()
 	: IsTraceRunningAttribute(*this)
+	, DynamicToolTipTextAttribute(*this)
 {
 }
 
@@ -44,6 +46,7 @@ void SToggleTraceButton::Construct(const FArguments& InArgs)
 {
 	OnToggleTraceRequested = InArgs._OnToggleTraceRequested;
 	IsTraceRunningAttribute.Assign(*this, InArgs._IsTraceRunning);
+	DynamicToolTipTextAttribute.Assign(*this, InArgs._DynamicToolTipText);
 	ButtonSize = InArgs._ButtonSize;
 
 	ChildSlot
@@ -132,12 +135,18 @@ FSlateColor SToggleTraceButton::GetRecordingButtonOutlineColor() const
 
 FText SToggleTraceButton::GetRecordingButtonTooltipText() const
 {
-	if (!IsTraceRunningAttribute.Get())
+	FText ToolTipText = DynamicToolTipTextAttribute.Get();
+	if (ToolTipText.IsEmpty())
 	{
-		return LOCTEXT("StartTracing", "Start tracing. The trace destination is set from the menu.");
+		if (!IsTraceRunningAttribute.Get())
+		{
+			return LOCTEXT("StartTracing", "Start tracing. The trace destination is set from the menu.");
+		}
+
+		return LOCTEXT("StopTracing", "Stop Tracing.");
 	}
 
-	return LOCTEXT("StopTracing", "Stop Tracing.");
+	return ToolTipText;
 }
 
 void SToggleTraceButton::ToggleTrace_OnClicked() const
diff --git a/Engine/Source/Developer/TraceTools/Public/TraceTools/Widgets/SToggleTraceButton.h b/Engine/Source/Developer/TraceTools/Public/TraceTools/Widgets/SToggleTraceButton.h
--- a/Engine/Source/Developer/TraceTools/Public/TraceTools/Widgets/SToggleTraceButton.h
+++ b/Engine/Source/Developer/TraceTools/Public/TraceTools/Widgets/SToggleTraceButton.h
@@ -37,6 +37,9 @@ public:
 	
 	/* Specifies if trace is running at the moment of call. */
 	SLATE_ATTRIBUTE(bool, IsTraceRunning)
+
+	/* Specifies the tooltip text that should change based on trace status. */
+	SLATE_ATTRIBUTE(FText, DynamicToolTipText)
 	
 	/* Specifies the size of the button from a few presets. */
 	SLATE_ARGUMENT(EButtonSize, ButtonSize)
@@ -70,6 +73,7 @@ private:
 	mutable double ConnectionStartTime = 0.0f;
 
 	TSlateAttribute<bool> IsTraceRunningAttribute;
+	TSlateAttribute<FText> DynamicToolTipTextAttribute;
 	EButtonSize ButtonSize;
 };
 
