[Insights] Updated support for Relay and Direct Trace connections for the TraceTools and TraceController.

diff --git a/Engine/Source/Developer/TraceTools/Private/Widgets/STraceStatistics.cpp b/Engine/Source/Developer/TraceTools/Private/Widgets/STraceStatistics.cpp
--- a/Engine/Source/Developer/TraceTools/Private/Widgets/STraceStatistics.cpp
+++ b/Engine/Source/Developer/TraceTools/Private/Widgets/STraceStatistics.cpp
@@ -460,6 +460,7 @@ FText STraceStatistics::GetTraceSystemStateText() const
 	static_assert((uint8)FTraceStatus::ETraceSystemStatus::Available == (uint8) FTraceAuxiliary::ETraceSystemStatus::Available);
 	static_assert((uint8)FTraceStatus::ETraceSystemStatus::TracingToFile == (uint8) FTraceAuxiliary::ETraceSystemStatus::TracingToFile);
 	static_assert((uint8)FTraceStatus::ETraceSystemStatus::TracingToServer == (uint8) FTraceAuxiliary::ETraceSystemStatus::TracingToServer);
+	static_assert((uint8)FTraceStatus::ETraceSystemStatus::TracingToCustomRelay == (uint8)FTraceAuxiliary::ETraceSystemStatus::TracingToCustomRelay);
 	static_assert((uint8)FTraceStatus::ETraceSystemStatus::NumValues == (uint8) FTraceAuxiliary::ETraceSystemStatus::NumValues, "ETraceSystemStatus enum values are of out sync.");
 
 	if (!SessionFilterService->HasStats())
@@ -479,14 +480,18 @@ FText STraceStatistics::GetTraceSystemStateText() const
 	}
 	case FTraceStatus::ETraceSystemStatus::TracingToServer:
 	{
-		return LOCTEXT("TracingToServerText", "Tracing to Server");
+		return LOCTEXT("TracingToServerText", "Tracing to Network");
 	}
 	case FTraceStatus::ETraceSystemStatus::TracingToFile:
 	{
 		return LOCTEXT("TracingToFileText", "Tracing to File");
 	}
+	case FTraceStatus::ETraceSystemStatus::TracingToCustomRelay:
+	{
+		return LOCTEXT("TracingToCustomRelayText", "Tracing to Custom Relay");
+	}
 	default:
-		return LOCTEXT("Unkown", "Unkown");
+		return LOCTEXT("Unknown", "Unknown");
 	}
 }
 
@@ -509,12 +514,16 @@ FText STraceStatistics::GetTraceSystemStateTooltipText() const
 	}
 	case FTraceStatus::ETraceSystemStatus::TracingToServer:
 	{
-		return LOCTEXT("TracingToServerTooltipText", "Tracing to the trace server.");
+		return LOCTEXT("TracingToServerTooltipText", "Tracing to network (to trace server or using a direct trace connection).");
 	}
 	case FTraceStatus::ETraceSystemStatus::TracingToFile:
 	{
 		return LOCTEXT("TracingToFileTooltipText", "Tracing directly to a file.");
 	}
+	case FTraceStatus::ETraceSystemStatus::TracingToCustomRelay:
+	{
+		return LOCTEXT("TracingToCustomRelayTooltipText", "Tracing to a custom relay.");
+	}
 	default:
 		return FText::GetEmpty();
 	}
diff --git a/Engine/Source/Runtime/Core/Private/ProfilingDebugging/TraceAuxiliary.cpp b/Engine/Source/Runtime/Core/Private/ProfilingDebugging/TraceAuxiliary.cpp
--- a/Engine/Source/Runtime/Core/Private/ProfilingDebugging/TraceAuxiliary.cpp
+++ b/Engine/Source/Runtime/Core/Private/ProfilingDebugging/TraceAuxiliary.cpp
@@ -2314,6 +2314,10 @@ FTraceAuxiliary::ETraceSystemStatus FTraceAuxiliary::GetTraceSystemStatus()
 	{
 		return ETraceSystemStatus::TracingToFile;
 	}
+	else if (ConnectionType == EConnectionType::Relay)
+	{
+		return ETraceSystemStatus::TracingToCustomRelay;
+	}
 	else
 	{
 		return ETraceSystemStatus::Available;
diff --git a/Engine/Source/Runtime/Core/Public/ProfilingDebugging/TraceAuxiliary.h b/Engine/Source/Runtime/Core/Public/ProfilingDebugging/TraceAuxiliary.h
--- a/Engine/Source/Runtime/Core/Public/ProfilingDebugging/TraceAuxiliary.h
+++ b/Engine/Source/Runtime/Core/Public/ProfilingDebugging/TraceAuxiliary.h
@@ -31,6 +31,7 @@ public:
 		Available,
 		TracingToServer,
 		TracingToFile,
+		TracingToCustomRelay,
 
 		NumValues, // This must be the last value.
 	};
diff --git a/Engine/Source/Runtime/SessionServices/Public/ITraceController.h b/Engine/Source/Runtime/SessionServices/Public/ITraceController.h
--- a/Engine/Source/Runtime/SessionServices/Public/ITraceController.h
+++ b/Engine/Source/Runtime/SessionServices/Public/ITraceController.h
@@ -33,6 +33,7 @@ struct FTraceStatus
 		Available,
 		TracingToServer,
 		TracingToFile,
+		TracingToCustomRelay,
 
 		NumValues,
 	};
