Move trace control buttons to the top toolbar in Verse Insights.
Add label to the Toggle Trace button in Verse Insights.
Update Verse Insights layout.

#FailedValidations Preflight

diff --git a/Engine/Source/Developer/TraceInsights/Private/Insights/TimingProfiler/Widgets/STimingProfilerToolbar.cpp b/Engine/Source/Developer/TraceInsights/Private/Insights/TimingProfiler/Widgets/STimingProfilerToolbar.cpp
--- a/Engine/Source/Developer/TraceInsights/Private/Insights/TimingProfiler/Widgets/STimingProfilerToolbar.cpp
+++ b/Engine/Source/Developer/TraceInsights/Private/Insights/TimingProfiler/Widgets/STimingProfilerToolbar.cpp
@@ -39,6 +39,11 @@ void STimingProfilerToolbar::Construct(const FArguments& InArgs, const FInsights
 	{
 		static void FillViewToolbar(FToolBarBuilder& ToolbarBuilder, const FArguments &InArgs, const FInsightsMajorTabConfig& Config)
 		{
+			if (InArgs._ToolbarExtender.IsValid())
+			{
+				InArgs._ToolbarExtender->Apply("BeforeMainToolbar", EExtensionHook::First, ToolbarBuilder);
+			}
+
 			ToolbarBuilder.BeginSection("View");
 			{
 				if (Config.ShouldRegisterMinorTab(FTimingProfilerTabs::FramesTrackID))
diff --git a/Engine/Source/Developer/TraceTools/Private/Widgets/SToggleTraceButton.cpp b/Engine/Source/Developer/TraceTools/Private/Widgets/SToggleTraceButton.cpp
--- a/Engine/Source/Developer/TraceTools/Private/Widgets/SToggleTraceButton.cpp
+++ b/Engine/Source/Developer/TraceTools/Private/Widgets/SToggleTraceButton.cpp
@@ -5,6 +5,7 @@
 #include "Framework/Application/SlateApplication.h"
 #include "Internationalization/Text.h"
 #include "SlateOptMacros.h"
+#include "Styling/StarshipCoreStyle.h"
 #include "Styling/StyleColors.h"
 #include "Widgets/Images/SImage.h"
 #include "Widgets/Input/SButton.h"
@@ -29,11 +30,13 @@ void SToggleTraceButton::PrivateRegisterAttributes(FSlateAttributeInitializer& A
 {
 	SLATE_ADD_MEMBER_ATTRIBUTE_DEFINITION_WITH_NAME(AttributeInitializer, "IsTraceRunning", IsTraceRunningAttribute, EInvalidateWidgetReason::Paint);
 	SLATE_ADD_MEMBER_ATTRIBUTE_DEFINITION_WITH_NAME(AttributeInitializer, "DynamicToolTipText", DynamicToolTipTextAttribute, EInvalidateWidgetReason::Paint);
+	SLATE_ADD_MEMBER_ATTRIBUTE_DEFINITION_WITH_NAME(AttributeInitializer, "LabelText", LabelTextAttribute, EInvalidateWidgetReason::Paint);
 }
 
 SToggleTraceButton::SToggleTraceButton()
 	: IsTraceRunningAttribute(*this)
 	, DynamicToolTipTextAttribute(*this)
+	, LabelTextAttribute(*this)
 {
 }
 
@@ -47,6 +50,7 @@ void SToggleTraceButton::Construct(const FArguments& InArgs)
 	OnToggleTraceRequested = InArgs._OnToggleTraceRequested;
 	IsTraceRunningAttribute.Assign(*this, InArgs._IsTraceRunning);
 	DynamicToolTipTextAttribute.Assign(*this, InArgs._DynamicToolTipText);
+	LabelTextAttribute.Assign(*this, InArgs._LabelText);
 	ButtonSize = InArgs._ButtonSize;
 
 	ChildSlot
@@ -62,29 +66,41 @@ void SToggleTraceButton::Construct(const FArguments& InArgs)
 		.OnUnhovered_Lambda([this]() { this->bIsTraceRecordButtonHovered = false; })
 		.Content()
 		[
-			SNew(SOverlay)
+			SNew(SHorizontalBox)
 
-			+ SOverlay::Slot()
+			+ SHorizontalBox::Slot()
 			[
-				SNew(SImage)
-				.ColorAndOpacity(this, &SToggleTraceButton::GetRecordingButtonColor)
-				.Image(GetToggleTraceCenterBrush())
-				.Visibility(this, &SToggleTraceButton::GetStartTraceIconVisibility)
+				SNew(SOverlay)
+
+				+ SOverlay::Slot()
+				[
+					SNew(SImage)
+					.ColorAndOpacity(this, &SToggleTraceButton::GetRecordingButtonColor)
+					.Image(GetToggleTraceCenterBrush())
+					.Visibility(this, &SToggleTraceButton::GetStartTraceIconVisibility)
+				]
+
+				+ SOverlay::Slot()
+				[
+					SNew(SImage)
+					.ColorAndOpacity(this, &SToggleTraceButton::GetRecordingButtonOutlineColor)
+					.Image(GetToggleTraceOutlineBrush())
+					.Visibility(this, &SToggleTraceButton::GetStartTraceIconVisibility)
+				]
+
+				+ SOverlay::Slot()
+				[
+					SNew(SImage)
+					.Image(GetToggleTraceStopBrush())
+					.Visibility(this, &SToggleTraceButton::GetStopTraceIconVisibility)
+				]
 			]
-
-			+ SOverlay::Slot()
-			[
-				SNew(SImage)
-				.ColorAndOpacity(this, &SToggleTraceButton::GetRecordingButtonOutlineColor)
-				.Image(GetToggleTraceOutlineBrush())
-				.Visibility(this, &SToggleTraceButton::GetStartTraceIconVisibility)
-			]
-
-			+ SOverlay::Slot()
+			+ SHorizontalBox::Slot()
+			.AutoWidth()
+			.Padding(FMargin(4.0f, 2.0f, 0.0f, 0.0f))
 			[
-				SNew(SImage)
-				.Image(GetToggleTraceStopBrush())
-				.Visibility(this, &SToggleTraceButton::GetStopTraceIconVisibility)
+				SNew(STextBlock)
+				.Text(this, &SToggleTraceButton::GetRecordingButtonLabelText)
 			]
 		]
 	];
@@ -149,6 +165,11 @@ FText SToggleTraceButton::GetRecordingButtonTooltipText() const
 	return ToolTipText;
 }
 
+FText SToggleTraceButton::GetRecordingButtonLabelText() const
+{
+	return LabelTextAttribute.Get();
+}
+
 void SToggleTraceButton::ToggleTrace_OnClicked() const
 {
 	OnToggleTraceRequested.ExecuteIfBound();
diff --git a/Engine/Source/Developer/TraceTools/Public/TraceTools/Widgets/SToggleTraceButton.h b/Engine/Source/Developer/TraceTools/Public/TraceTools/Widgets/SToggleTraceButton.h
--- a/Engine/Source/Developer/TraceTools/Public/TraceTools/Widgets/SToggleTraceButton.h
+++ b/Engine/Source/Developer/TraceTools/Public/TraceTools/Widgets/SToggleTraceButton.h
@@ -40,6 +40,9 @@ public:
 
 	/* Specifies the tooltip text that should change based on trace status. */
 	SLATE_ATTRIBUTE(FText, DynamicToolTipText)
+
+	/* Label to place in the button. */
+	SLATE_ATTRIBUTE(FText, LabelText)
 	
 	/* Specifies the size of the button from a few presets. */
 	SLATE_ARGUMENT(EButtonSize, ButtonSize)
@@ -56,6 +59,7 @@ private:
 	FSlateColor GetRecordingButtonColor() const;
 	FSlateColor GetRecordingButtonOutlineColor() const;
 	FText GetRecordingButtonTooltipText() const;
+	FText GetRecordingButtonLabelText() const;
 
 	EVisibility GetStartTraceIconVisibility() const;
 	EVisibility GetStopTraceIconVisibility() const;
@@ -74,6 +78,7 @@ private:
 
 	TSlateAttribute<bool> IsTraceRunningAttribute;
 	TSlateAttribute<FText> DynamicToolTipTextAttribute;
+	TSlateAttribute<FText> LabelTextAttribute;
 	EButtonSize ButtonSize;
 };
 
