[Insights] Add status enums for secure connections.

#Synced-CL 50345636

diff --git a/Engine/Source/Developer/TraceInsights/Private/Insights/InsightsManager.cpp b/Engine/Source/Developer/TraceInsights/Private/Insights/InsightsManager.cpp
--- a/Engine/Source/Developer/TraceInsights/Private/Insights/InsightsManager.cpp
+++ b/Engine/Source/Developer/TraceInsights/Private/Insights/InsightsManager.cpp
@@ -909,7 +909,7 @@ void FInsightsManager::ConnectToSecureTrace(FStringView Host, uint16 Port, TCons
 
 		if (Session)
 		{
-			CurrentTraceStreamType = ETraceStreamType::DirectTrace;
+			CurrentTraceStreamType = ETraceStreamType::DirectTraceSecure;
 			CurrentTraceName = TraceName;
 			CurrentTraceId = 0;
 			CurrentDirectTracePort = Port;
diff --git a/Engine/Source/Developer/TraceInsights/Private/Insights/InsightsManager.h b/Engine/Source/Developer/TraceInsights/Private/Insights/InsightsManager.h
--- a/Engine/Source/Developer/TraceInsights/Private/Insights/InsightsManager.h
+++ b/Engine/Source/Developer/TraceInsights/Private/Insights/InsightsManager.h
@@ -65,6 +65,7 @@ enum class ETraceStreamType
 	TraceFile,
 	TraceStore,
 	DirectTrace,
+	DirectTraceSecure,
 	Custom
 };
 
diff --git a/Engine/Source/Developer/TraceInsights/Private/Insights/Widgets/SSessionInfoWindow.cpp b/Engine/Source/Developer/TraceInsights/Private/Insights/Widgets/SSessionInfoWindow.cpp
--- a/Engine/Source/Developer/TraceInsights/Private/Insights/Widgets/SSessionInfoWindow.cpp
+++ b/Engine/Source/Developer/TraceInsights/Private/Insights/Widgets/SSessionInfoWindow.cpp
@@ -408,6 +408,7 @@ void SSessionInfoWindow::UpdateStreamTypeAndSessionName(const TraceServices::IAn
 		break;
 
 	case ETraceStreamType::DirectTrace:
+	case ETraceStreamType::DirectTraceSecure:
 		StreamTypeText = FText::Format(LOCTEXT("TraceStreamType_DirectTrace_Fmt", "Direct Trace -- Port={0}"),
 			FText::AsNumber(uint32(FInsightsManager::Get()->GetDirectTracePort()), &FNumberFormattingOptions::DefaultNoGrouping()));
 		break;
diff --git a/Engine/Source/Developer/TraceTools/Private/Widgets/STraceStatistics.cpp b/Engine/Source/Developer/TraceTools/Private/Widgets/STraceStatistics.cpp
--- a/Engine/Source/Developer/TraceTools/Private/Widgets/STraceStatistics.cpp
+++ b/Engine/Source/Developer/TraceTools/Private/Widgets/STraceStatistics.cpp
@@ -461,6 +461,7 @@ FText STraceStatistics::GetTraceSystemStateText() const
 	static_assert((uint8)FTraceStatus::ETraceSystemStatus::TracingToFile == (uint8) FTraceAuxiliary::ETraceSystemStatus::TracingToFile);
 	static_assert((uint8)FTraceStatus::ETraceSystemStatus::TracingToServer == (uint8) FTraceAuxiliary::ETraceSystemStatus::TracingToServer);
 	static_assert((uint8)FTraceStatus::ETraceSystemStatus::TracingToCustomRelay == (uint8)FTraceAuxiliary::ETraceSystemStatus::TracingToCustomRelay);
+	static_assert((uint8)FTraceStatus::ETraceSystemStatus::TracingToSecureNetwork == (uint8)FTraceAuxiliary::ETraceSystemStatus::TracingToSecureNetwork);
 	static_assert((uint8)FTraceStatus::ETraceSystemStatus::NumValues == (uint8) FTraceAuxiliary::ETraceSystemStatus::NumValues, "ETraceSystemStatus enum values are of out sync.");
 
 	if (!SessionFilterService->HasStats())
@@ -490,6 +491,10 @@ FText STraceStatistics::GetTraceSystemStateText() const
 	{
 		return LOCTEXT("TracingToCustomRelayText", "Tracing to Custom Relay");
 	}
+	case FTraceStatus::ETraceSystemStatus::TracingToSecureNetwork:
+	{
+		return LOCTEXT("TracingToSecureNetwork", "Tracing to secure network connection");
+	}
 	default:
 		return LOCTEXT("Unknown", "Unknown");
 	}
@@ -524,6 +529,10 @@ FText STraceStatistics::GetTraceSystemStateTooltipText() const
 	{
 		return LOCTEXT("TracingToCustomRelayTooltipText", "Tracing to a custom relay.");
 	}
+	case FTraceStatus::ETraceSystemStatus::TracingToSecureNetwork:
+	{
+		return LOCTEXT("TracingToSecureNetwork", "Tracing to secure network connection");
+	}
 	default:
 		return FText::GetEmpty();
 	}
@@ -553,4 +562,4 @@ EVisibility STraceStatistics::GetCopyEndpointVisibility() const
 
 } // namespace UE::TraceTools
 
-#undef LOCTEXT_NAMESPACE
\ No newline at end of file
+#undef LOCTEXT_NAMESPACE
diff --git a/Engine/Source/Runtime/Core/Public/ProfilingDebugging/TraceAuxiliary.h b/Engine/Source/Runtime/Core/Public/ProfilingDebugging/TraceAuxiliary.h
--- a/Engine/Source/Runtime/Core/Public/ProfilingDebugging/TraceAuxiliary.h
+++ b/Engine/Source/Runtime/Core/Public/ProfilingDebugging/TraceAuxiliary.h
@@ -32,6 +32,7 @@ public:
 		Available,
 		TracingToServer,
 		TracingToFile,
+		TracingToSecureNetwork,
 		TracingToCustomRelay,
 
 		NumValues, // This must be the last value.
@@ -90,12 +91,10 @@ public:
 		 * Relay connection. Pass user defined I/O functions to write raw trace data.
 		 */
 		Relay,
-#if UE_TRACE_ALLOW_SECURE_TRACING
 		/**
 		 * Connect to a receiver using a secure socket connection.
 		 */
 		SecureNetwork,
-#endif
 		/**
 		 * Don't connect, just start tracing to memory.
 		 */
diff --git a/Engine/Source/Runtime/SessionServices/Public/ITraceController.h b/Engine/Source/Runtime/SessionServices/Public/ITraceController.h
--- a/Engine/Source/Runtime/SessionServices/Public/ITraceController.h
+++ b/Engine/Source/Runtime/SessionServices/Public/ITraceController.h
@@ -33,6 +33,7 @@ struct FTraceStatus
 		Available,
 		TracingToServer,
 		TracingToFile,
+		TracingToSecureNetwork,
 		TracingToCustomRelay,
 
 		NumValues,
