Fix flickering for channel tooltips in the trace control wigdet.

#Synced-CL 52439714

diff --git a/Engine/Source/Developer/TraceTools/Private/Services/SessionTraceControllerFilterService.cpp b/Engine/Source/Developer/TraceTools/Private/Services/SessionTraceControllerFilterService.cpp
--- a/Engine/Source/Developer/TraceTools/Private/Services/SessionTraceControllerFilterService.cpp
+++ b/Engine/Source/Developer/TraceTools/Private/Services/SessionTraceControllerFilterService.cpp
@@ -150,24 +150,50 @@ void FSessionTraceControllerFilterService::OnTraceStatusUpdated(const FTraceStat
 
 void FSessionTraceControllerFilterService::UpdateChannels(const FTraceStatus& InStatus)
 {
-	ChannelsTimestamp = FDateTime::Now();
-
-	const TMap<uint32, FTraceStatus::FChannel> Channels = InStatus.Channels;
-	Objects.Empty(Channels.Num());
+	const TMap<uint32, FTraceStatus::FChannel>& Channels = InStatus.Channels;
 
 	if (Channels.Num())
 	{
 		bChannelsReceived = true;
 	}
 
-	for (auto& Entry : Channels)
+	bool bChanged = bHasPendingChannelChanges || Channels.Num() != Objects.Num();
+	bHasPendingChannelChanges = false;
+
+	TMap<uint64, FTraceObjectInfo> NewObjects;
+	NewObjects.Reserve(Channels.Num());
+
+	for (const auto& Entry : Channels)
 	{
-		FTraceObjectInfo& EventInfo = Objects.Add(HashName(Entry.Value.Name));
-		EventInfo.Name = Entry.Value.Name;
-		EventInfo.Description = Entry.Value.Description;
-		EventInfo.bEnabled = Entry.Value.bEnabled;
-		EventInfo.bReadOnly = Entry.Value.bReadOnly;
-		EventInfo.Id = Entry.Value.Id;
+		const uint64 Hash = HashName(Entry.Value.Name);
+		FTraceObjectInfo& NewInfo = NewObjects.Add(Hash);
+
+		if (!bChanged)
+		{
+			if (const FTraceObjectInfo* Existing = Objects.Find(Hash))
+			{
+				if (Existing->bEnabled != Entry.Value.bEnabled)
+				{
+					bChanged = true;
+				}
+			}
+			else
+			{
+				bChanged = true;
+			}
+		}
+
+		NewInfo.Name = Entry.Value.Name;
+		NewInfo.Description = Entry.Value.Description;
+		NewInfo.bEnabled = Entry.Value.bEnabled;
+		NewInfo.bReadOnly = Entry.Value.bReadOnly;
+		NewInfo.Id = Entry.Value.Id;
+	}
+
+	if (bChanged)
+	{
+		ChannelsTimestamp = FDateTime::Now();
+		Objects = MoveTemp(NewObjects);
 	}
 }
 
@@ -187,6 +213,8 @@ void FSessionTraceControllerFilterService::OnApplyChannelChanges()
 
 		FrameEnabledChannels.Empty();
 		FrameDisabledChannels.Empty();
+
+		bHasPendingChannelChanges = true;
 	}
 }
 
@@ -219,6 +247,7 @@ void FSessionTraceControllerFilterService::Reset()
 	TraceEndpoint.Empty();
 	ChannelsTimestamp = FDateTime();
 	StatusTimestamp = FDateTime();
+	bHasPendingChannelChanges = false;
 }
 
 void FSessionTraceControllerFilterService::SetInstanceId(const FGuid& Id)
diff --git a/Engine/Source/Developer/TraceTools/Private/Services/SessionTraceControllerFilterService.h b/Engine/Source/Developer/TraceTools/Private/Services/SessionTraceControllerFilterService.h
--- a/Engine/Source/Developer/TraceTools/Private/Services/SessionTraceControllerFilterService.h
+++ b/Engine/Source/Developer/TraceTools/Private/Services/SessionTraceControllerFilterService.h
@@ -86,6 +86,8 @@ protected:
 	FTraceStatus::ETraceSystemStatus TraceSystemStatus;
 
 	FGuid InstanceId;
+
+	bool bHasPendingChannelChanges = false;
 };
 
 } // namespace UE::TraceTools
\ No newline at end of file
